UE-199918 VisionOS OpenXR wrapper
-Wrap the VisionOS immersive mode api in and OpenXR layer so we can run VisionOS immersive mode through OpenXRHMD.
-Simulator renders in stereo and tracks the head correctly.  No on-device testing has been done and the simulator only provides one eye view, so we can't be sure it will display the second eye correctly, but I tried to set it up so it would.
-Very experimental.
-Known problems: Using the intermediate back buffer.  No features beyond basic head tracking and stereo rendering.  Lots of unused code for features that needs to be updated or deleted.

diff --git a/Engine/Platforms/VisionOS/Plugins/Runtime/OpenXRVisionOS/OpenXRVisionOS.uplugin b/Engine/Platforms/VisionOS/Plugins/Runtime/OpenXRVisionOS/OpenXRVisionOS.uplugin
new file mode 100644
--- /dev/null
+++ b/Engine/Platforms/VisionOS/Plugins/Runtime/OpenXRVisionOS/OpenXRVisionOS.uplugin
@@ -0,0 +1,47 @@
+{
+	"FileVersion": 3,
+	"Version": 1,
+	"VersionName": "1.0",
+	"FriendlyName": "OpenXRVisionOS",
+	"Description": "OpenXRVisionOS exposes VisionOS full immersion to Unreal as if it was an OpenXR runtime.  Documentation on developing VR in Unreal with OpenXR is applicable.",
+	"Category": "Virtual Reality",
+	"CreatedBy": "Epic Games, Inc.",
+	"CreatedByURL": "https://epicgames.com",
+	"DocsURL": "",
+	"MarketplaceURL": "",
+	"SupportURL": "",
+	"EnabledByDefault": false,
+	"CanContainContent": false,
+	"IsBetaVersion": true,
+	"IsExperimentalVersion": true,
+	"Installed": false,
+	"SupportedTargetPlatforms": [
+		"VisionOS"
+	],
+	"Modules": [
+		{
+			"Name": "OXRVisionOS",
+			"Type": "Runtime",
+			"LoadingPhase": "PostConfigInit",
+			"PlatformAllowList": [
+				"VisionOS"
+			]
+		},
+		{
+			"Name": "OXRVisionOSSettings",
+			"Type": "Runtime",
+			"LoadingPhase": "PreDefault",
+			"PlatformAllowList": [
+				"VisionOS",
+				"Win64",
+				"Mac"
+			]
+		}
+	],
+	"Plugins": [
+		{
+			"Name": "OpenXR",
+			"Enabled": true
+		}
+	]
+}
diff --git a/Engine/Platforms/VisionOS/Plugins/Runtime/OpenXRVisionOS/Source/OXRVisionOS/Internal/OXRVisionOS_openxr_platform.h b/Engine/Platforms/VisionOS/Plugins/Runtime/OpenXRVisionOS/Source/OXRVisionOS/Internal/OXRVisionOS_openxr_platform.h
new file mode 100644
--- /dev/null
+++ b/Engine/Platforms/VisionOS/Plugins/Runtime/OpenXRVisionOS/Source/OXRVisionOS/Internal/OXRVisionOS_openxr_platform.h
@@ -0,0 +1,37 @@
+// Copyright Epic Games, Inc. All Rights Reserved.
+
+#pragma once
+
+#include "openxr/openxr.h"
+#include "RHIFwd.h"
+
+// Uncomment to make intellisense work better.
+//#include "../../../../../../../../Source/ThirdParty/OpenXR/include/openxr/openxr.h"
+
+// This number is in a block that is unused by openxr, and that seemed unlikely to be used soon based on the patterns of usage.
+// But if it ever was used by openxr we would have to change it.
+#define XR_TYPE_OXRVISIONOS_EPIC 1000059100
+
+typedef enum XrOXRVisionOSStructureTypeEPIC {
+	XR_TYPE_GRAPHICS_BINDING_OXRVISIONOS_EPIC = XR_TYPE_OXRVISIONOS_EPIC + 0,
+	XR_TYPE_SWAPCHAIN_IMAGE_OXRVISIONOS_EPIC = XR_TYPE_OXRVISIONOS_EPIC + 1,
+} XrOXRVisionOSStructureTypeEPIC;
+
+typedef struct XrGraphicsBindingOXRVisionOSEPIC {
+	XrStructureType             type;
+	const void* XR_MAY_ALIAS    next;
+	// More members here
+} XrGraphicsBindingOXRVisionOSEPIC;                                                                                                                                                 
+
+typedef struct XrSwapchainImageOXRVisionOS {
+    XrStructureType       type;
+    void* XR_MAY_ALIAS    next;
+	FTextureRHIRef        image;
+} XrSwapchainImageOXRVisionOS;
+
+
+// OXRVisionOS Extensions
+
+#define XR_EPIC_oxrvisionos_controller 1
+#define XR_EPIC_oxrvisionos_controller_SPEC_VERSION 1
+#define XR_EPIC_OXRVISIONOS_CONTROLLER_NAME "XR_EPIC_oxrvisionos_controller"
\ No newline at end of file
diff --git a/Engine/Platforms/VisionOS/Plugins/Runtime/OpenXRVisionOS/Source/OXRVisionOS/OXRVisionOS.Build.cs b/Engine/Platforms/VisionOS/Plugins/Runtime/OpenXRVisionOS/Source/OXRVisionOS/OXRVisionOS.Build.cs
new file mode 100644
--- /dev/null
+++ b/Engine/Platforms/VisionOS/Plugins/Runtime/OpenXRVisionOS/Source/OXRVisionOS/OXRVisionOS.Build.cs
@@ -0,0 +1,33 @@
+// Copyright Epic Games, Inc. All Rights Reserved.
+
+using System;
+using System.IO;
+
+namespace UnrealBuildTool.Rules
+{
+	public class OXRVisionOS : ModuleRules
+	{
+        public OXRVisionOS(ReadOnlyTargetRules Target) : base(Target)
+		{
+			PublicDependencyModuleNames.AddRange(new string[]
+			{
+				"Core",
+				"CoreUObject",
+				"Engine",
+				"RenderCore",
+				"RHI",
+				"MetalRHI",
+				"OpenXR",
+				"OpenXRHMD",
+				"OXRVisionOSSettings",
+			});
+			
+		AddEngineThirdPartyPrivateStaticDependencies(Target,
+			"MTLPP"
+		);
+			PublicWeakFrameworks.Add("Metal");
+			PublicWeakFrameworks.Add("ARKit");		
+			PrivateIncludePathModuleNames.Add("OpenXRHMD");
+		}
+	}
+}
diff --git a/Engine/Platforms/VisionOS/Plugins/Runtime/OpenXRVisionOS/Source/OXRVisionOS/Private/OXRVisionOS.cpp b/Engine/Platforms/VisionOS/Plugins/Runtime/OpenXRVisionOS/Source/OXRVisionOS/Private/OXRVisionOS.cpp
new file mode 100644
--- /dev/null
+++ b/Engine/Platforms/VisionOS/Plugins/Runtime/OpenXRVisionOS/Source/OXRVisionOS/Private/OXRVisionOS.cpp
@@ -0,0 +1,926 @@
+// Copyright Epic Games, Inc. All Rights Reserved.
+
+#include "OXRVisionOS.h"
+#include "Engine/Engine.h"
+#include "Modules/ModuleManager.h"
+
+#include "OpenXRCore.h"
+#include "OXRVisionOSInstance.h"
+#include "OXRVisionOSSession.h"
+#include "OXRVisionOSSpace.h"
+#include "OXRVisionOSActionSet.h"
+#include "OXRVisionOSAction.h"
+#include "OXRVisionOS_RenderBridge.h"
+#include "OXRVisionOSSwapchain.h"
+#include "OXRVisionOSPlatformUtils.h"
+#include "OXRVisionOS_openxr_platform.h"
+
+
+DEFINE_LOG_CATEGORY(LogOXRVisionOS);
+
+TSharedPtr<FOXRVisionOSInstance, ESPMode::ThreadSafe> FOXRVisionOS::Instance;
+
+#define OXRVISIONOS_CHECK_MODULE_EARLY_RETURN(Module) \
+check(Module); \
+if (Module == nullptr) \
+{ \
+	return XrResult::XR_ERROR_RUNTIME_FAILURE; \
+}
+
+#define OXRVISIONOS_CHECK_INSTANCE_EARLY_RETURN(Instance) \
+if (!FOXRVisionOS::CheckInstance(Instance)) \
+{ \
+	return XrResult::XR_ERROR_HANDLE_INVALID; \
+}
+
+#define OXRVISIONOS_CHECK_SESSION_EARLY_RETURN(Session) \
+if (!FOXRVisionOSInstance::CheckSession(Session)) \
+{ \
+	return XrResult::XR_ERROR_HANDLE_INVALID; \
+}
+
+#define OXRVISIONOS_CHECK_SWAPCHAIN_EARLY_RETURN(SwapChain) \
+if (SwapChain == XR_NULL_HANDLE) \
+{ \
+	return XrResult::XR_ERROR_HANDLE_INVALID; \
+}
+
+#if 0
+#define OXRVISIONOS_LOG_RESULT(Result) \
+{ \
+	UE_LOG(LogOXRVisionOS, Verbose, TEXT("%s: %s"), ANSI_TO_TCHAR(__FUNCTION__), OpenXRResultToString(Result)); \
+}
+#else
+#define OXRVISIONOS_LOG_RESULT(Result)
+#endif
+
+
+namespace OXRVisionOS
+{
+	TMap<FString, PFN_xrVoidFunction> GlobalFunctionMap;
+	TMap<FString, PFN_xrVoidFunction> InstanceFunctionMap;
+
+	XRAPI_ATTR XrResult XRAPI_CALL xrGetInstanceProcAddr(
+		XrInstance                                  instance,
+		const char*									name,
+		PFN_xrVoidFunction*							function)
+	{
+		if (instance == XR_NULL_HANDLE)
+		{
+			PFN_xrVoidFunction* itr = GlobalFunctionMap.Find(name);
+			if (itr)
+			{
+				*function = *itr;
+				return XrResult::XR_SUCCESS;
+			}
+			else
+			{
+				return XrResult::XR_ERROR_FUNCTION_UNSUPPORTED;
+			}
+
+		}
+		else
+		{
+			PFN_xrVoidFunction* itr = InstanceFunctionMap.Find(name);
+			if (itr)
+			{
+				*function = *itr;
+				return XrResult::XR_SUCCESS;
+			}
+			else
+			{
+				return XrResult::XR_ERROR_FUNCTION_UNSUPPORTED;
+			}
+		}
+	}
+
+
+	// Global Functions
+
+	XRAPI_ATTR XrResult XRAPI_CALL xrEnumerateApiLayerProperties(
+		uint32_t                                    propertyCapacityInput,
+		uint32_t*									propertyCountOutput,
+		XrApiLayerProperties*						properties)
+	{
+		FOXRVisionOS* Module = FOXRVisionOS::Get();
+		OXRVISIONOS_CHECK_MODULE_EARLY_RETURN(Module);
+
+		return Module->XrEnumerateApiLayerProperties(
+			propertyCapacityInput,
+			propertyCountOutput,
+			properties);
+	}
+
+	XRAPI_ATTR XrResult XRAPI_CALL xrEnumerateInstanceExtensionProperties(
+		const char*									layerName,
+		uint32_t                                    propertyCapacityInput,
+		uint32_t*									propertyCountOutput,
+		XrExtensionProperties*						properties)
+	{
+		FOXRVisionOS* Module = FOXRVisionOS::Get();
+		OXRVISIONOS_CHECK_MODULE_EARLY_RETURN(Module);
+
+		return Module->XrEnumerateInstanceExtensionProperties(
+			layerName,
+			propertyCapacityInput,
+			propertyCountOutput,
+			properties);
+	}
+
+	XRAPI_ATTR XrResult XRAPI_CALL xrCreateInstance(
+		const XrInstanceCreateInfo*					createInfo,
+		XrInstance*									instance)
+	{
+		FOXRVisionOS* Module = FOXRVisionOS::Get();
+		OXRVISIONOS_CHECK_MODULE_EARLY_RETURN(Module);
+
+		XrResult Result = Module->XrCreateInstance(createInfo, instance);
+		OXRVISIONOS_LOG_RESULT(Result);
+		return Result;
+	}
+
+
+	// Instance functions
+
+	XRAPI_ATTR XrResult XRAPI_CALL xrDestroyInstance(
+		XrInstance                                  instance)
+	{
+		FOXRVisionOS* Module = FOXRVisionOS::Get();
+		check(Module);
+		if (Module == nullptr)
+		{
+			check(false);
+			return XrResult::XR_ERROR_HANDLE_INVALID; // This is the only allowed failure code. Not really appropriate, but this should never happen.
+		}
+
+		XrResult Result = Module->XrDestroyInstance(instance);
+		OXRVISIONOS_LOG_RESULT(Result);
+		return Result;
+	}
+
+	XRAPI_ATTR XrResult XRAPI_CALL xrGetInstanceProperties(
+		XrInstance                                  instance,
+		XrInstanceProperties* instanceProperties)
+	{
+		OXRVISIONOS_CHECK_INSTANCE_EARLY_RETURN(instance);
+
+		FOXRVisionOSInstance* Instance = (FOXRVisionOSInstance*)instance;
+		return Instance->XrGetInstanceProperties(instanceProperties);
+	}
+
+	XRAPI_ATTR XrResult XRAPI_CALL xrPollEvent(
+		XrInstance                                  instance,
+		XrEventDataBuffer*							eventData)
+	{
+		OXRVISIONOS_CHECK_INSTANCE_EARLY_RETURN(instance);
+
+		FOXRVisionOSInstance* Instance = (FOXRVisionOSInstance*)instance;
+		return Instance->XrPollEvent(eventData);
+	}
+
+#	define OXRVISIONOS_XR_RESULT_STR(name, val) case name: FCStringAnsi::Strncpy(buffer, #name, XR_MAX_RESULT_STRING_SIZE); return XrResult::XR_SUCCESS;
+
+	XRAPI_ATTR XrResult XRAPI_CALL xrResultToString(
+		XrInstance                                  instance,
+		XrResult                                    value,
+		char                                        buffer[XR_MAX_RESULT_STRING_SIZE])
+	{
+		OXRVISIONOS_CHECK_INSTANCE_EARLY_RETURN(instance);
+		switch (value)
+		{
+			XR_LIST_ENUM_XrResult(OXRVISIONOS_XR_RESULT_STR);
+			default: 
+			{
+				if (value > 0)
+				{
+					FCStringAnsi::Snprintf(buffer, XR_MAX_RESULT_STRING_SIZE, "XR_UNKNOWN_SUCCESS_%i", int(value));
+				}
+				else
+				{
+					FCStringAnsi::Snprintf(buffer, XR_MAX_RESULT_STRING_SIZE, "XR_UNKNOWN_FAILURE_%i", int(value));
+				}
+				return XrResult::XR_SUCCESS;
+			}
+		}
+	}
+
+#	define OXRVISIONOS_XR_STRUCTURETYPE_STR(name, val) case name: FCStringAnsi::Strncpy(buffer, #name, XR_MAX_STRUCTURE_NAME_SIZE); return XrResult::XR_SUCCESS;
+
+	XRAPI_ATTR XrResult XRAPI_CALL xrStructureTypeToString(
+		XrInstance                                  instance,
+		XrStructureType                             value,
+		char                                        buffer[XR_MAX_STRUCTURE_NAME_SIZE])
+	{
+		OXRVISIONOS_CHECK_INSTANCE_EARLY_RETURN(instance);
+		switch (value)
+		{
+			XR_LIST_ENUM_XrStructureType(OXRVISIONOS_XR_STRUCTURETYPE_STR);
+			default:
+			{
+				FCStringAnsi::Snprintf(buffer, XR_MAX_STRUCTURE_NAME_SIZE, "XR_UNKNOWN_STRUCTURE_TYPE_ %i", int(value));
+				return XrResult::XR_SUCCESS;
+			}
+		}
+	}
+
+	XRAPI_ATTR XrResult XRAPI_CALL xrGetSystem(
+		XrInstance                                  instance,
+		const XrSystemGetInfo*						getInfo,
+		XrSystemId*									systemId)
+	{
+		OXRVISIONOS_CHECK_INSTANCE_EARLY_RETURN(instance);
+		FOXRVisionOSInstance* Instance = (FOXRVisionOSInstance*)instance;
+
+		XrResult Result = Instance->XrGetSystem(getInfo, systemId);
+		OXRVISIONOS_LOG_RESULT(Result);
+		return Result;
+	}
+
+	XRAPI_ATTR XrResult XRAPI_CALL xrGetSystemProperties(
+		XrInstance                                  instance,
+		XrSystemId                                  systemId,
+		XrSystemProperties*							properties)
+	{
+		OXRVISIONOS_CHECK_INSTANCE_EARLY_RETURN(instance);
+		FOXRVisionOSInstance* Instance = (FOXRVisionOSInstance*)instance;
+		return Instance->XrGetSystemProperties(systemId, properties);
+	}
+
+	XRAPI_ATTR XrResult XRAPI_CALL xrEnumerateEnvironmentBlendModes(
+		XrInstance                                  instance,
+		XrSystemId                                  systemId,
+		XrViewConfigurationType                     viewConfigurationType,
+		uint32_t                                    environmentBlendModeCapacityInput,
+		uint32_t*									environmentBlendModeCountOutput,
+		XrEnvironmentBlendMode*						environmentBlendModes)
+	{
+		OXRVISIONOS_CHECK_INSTANCE_EARLY_RETURN(instance);
+		FOXRVisionOSInstance* Instance = (FOXRVisionOSInstance*)instance;
+		return Instance->XrEnumerateEnvironmentBlendModes(	systemId,
+															viewConfigurationType,
+															environmentBlendModeCapacityInput,
+															environmentBlendModeCountOutput,
+															environmentBlendModes);
+	}
+
+	XRAPI_ATTR XrResult XRAPI_CALL xrCreateSession(
+		XrInstance                                  instance,
+		const XrSessionCreateInfo*					createInfo,
+		XrSession*									session)
+	{
+		OXRVISIONOS_CHECK_INSTANCE_EARLY_RETURN(instance);
+		FOXRVisionOSInstance* Instance = (FOXRVisionOSInstance*)instance;
+		XrResult Result = Instance->XrCreateSession(createInfo, session);
+		OXRVISIONOS_LOG_RESULT(Result);
+		return Result;
+	}
+
+	XRAPI_ATTR XrResult XRAPI_CALL xrDestroySession(
+		XrSession                                   session)
+	{
+		OXRVISIONOS_CHECK_SESSION_EARLY_RETURN(session);
+
+		FOXRVisionOSSession* Session = (FOXRVisionOSSession*)session;
+		XrResult Result = Session->XrDestroySession();
+		OXRVISIONOS_LOG_RESULT(Result);
+		return Result;
+	}
+
+	XRAPI_ATTR XrResult XRAPI_CALL xrEnumerateReferenceSpaces(
+		XrSession                                   session,
+		uint32_t                                    spaceCapacityInput,
+		uint32_t*									spaceCountOutput,
+		XrReferenceSpaceType*						spaces)
+	{
+		OXRVISIONOS_CHECK_SESSION_EARLY_RETURN(session);
+
+		FOXRVisionOSSession* Session = (FOXRVisionOSSession*)session;
+		return Session->XrEnumerateReferenceSpaces(	spaceCapacityInput,
+													spaceCountOutput,
+													spaces);
+	}
+
+	XRAPI_ATTR XrResult XRAPI_CALL xrCreateReferenceSpace(
+		XrSession                                   session,
+		const XrReferenceSpaceCreateInfo*			createInfo,
+		XrSpace*									space)
+	{
+		OXRVISIONOS_CHECK_SESSION_EARLY_RETURN(session);
+
+		FOXRVisionOSSession* Session = (FOXRVisionOSSession*)session;
+		return Session->XrCreateReferenceSpace(	createInfo,
+												space);
+	}
+
+	XRAPI_ATTR XrResult XRAPI_CALL xrGetReferenceSpaceBoundsRect(
+		XrSession                                   session,
+		XrReferenceSpaceType                        referenceSpaceType,
+		XrExtent2Df*								bounds)
+	{
+		OXRVISIONOS_CHECK_SESSION_EARLY_RETURN(session);
+
+		FOXRVisionOSSession* Session = (FOXRVisionOSSession*)session;
+		return Session->XrGetReferenceSpaceBoundsRect(	referenceSpaceType,
+														bounds);
+	}
+
+	XRAPI_ATTR XrResult XRAPI_CALL xrCreateActionSpace(
+		XrSession                                   session,
+		const XrActionSpaceCreateInfo*				createInfo,
+		XrSpace*									space)
+	{
+		OXRVISIONOS_CHECK_SESSION_EARLY_RETURN(session);
+
+		FOXRVisionOSSession* Session = (FOXRVisionOSSession*)session;
+		return Session->XrCreateActionSpace(createInfo,
+			space);
+	}
+
+	XRAPI_ATTR XrResult XRAPI_CALL xrLocateSpace(
+		XrSpace                                     space,
+		XrSpace                                     baseSpace,
+		XrTime                                      time,
+		XrSpaceLocation*							location)
+	{
+		FOXRVisionOSSpace* Space = (FOXRVisionOSSpace*)space;
+		return Space->XrLocateSpace(baseSpace,
+			time,
+			location);
+	}
+
+	XRAPI_ATTR XrResult XRAPI_CALL xrDestroySpace(
+		XrSpace                                     space)
+	{
+		FOXRVisionOSSpace* Space = (FOXRVisionOSSpace*)space;
+		return Space->XrDestroySpace();
+	}
+
+	XRAPI_ATTR XrResult XRAPI_CALL xrEnumerateViewConfigurations(
+		XrInstance                                  instance,
+		XrSystemId                                  systemId,
+		uint32_t                                    viewConfigurationTypeCapacityInput,
+		uint32_t*									viewConfigurationTypeCountOutput,
+		XrViewConfigurationType*					viewConfigurationTypes)
+	{
+		OXRVISIONOS_CHECK_INSTANCE_EARLY_RETURN(instance);
+		FOXRVisionOSInstance* Instance = (FOXRVisionOSInstance*)instance;
+		return Instance->XrEnumerateViewConfigurations(	systemId,
+														viewConfigurationTypeCapacityInput,
+														viewConfigurationTypeCountOutput,
+														viewConfigurationTypes);
+	}
+
+	XRAPI_ATTR XrResult XRAPI_CALL xrGetViewConfigurationProperties(
+		XrInstance                                  instance,
+		XrSystemId                                  systemId,
+		XrViewConfigurationType                     viewConfigurationType,
+		XrViewConfigurationProperties*				configurationProperties)
+	{
+		OXRVISIONOS_CHECK_INSTANCE_EARLY_RETURN(instance);
+		FOXRVisionOSInstance* Instance = (FOXRVisionOSInstance*)instance;
+		return Instance->XrGetViewConfigurationProperties(	systemId,
+															viewConfigurationType,
+															configurationProperties);
+	}
+
+	XRAPI_ATTR XrResult XRAPI_CALL xrEnumerateViewConfigurationViews(
+		XrInstance                                  instance,
+		XrSystemId                                  systemId,
+		XrViewConfigurationType                     viewConfigurationType,
+		uint32_t                                    viewCapacityInput,
+		uint32_t*									viewCountOutput,
+		XrViewConfigurationView*					views)
+	{
+		OXRVISIONOS_CHECK_INSTANCE_EARLY_RETURN(instance);
+		FOXRVisionOSInstance* Instance = (FOXRVisionOSInstance*)instance;
+		return Instance->XrEnumerateViewConfigurationViews(	systemId,
+															viewConfigurationType,
+															viewCapacityInput,
+															viewCountOutput,
+															views);
+	}
+
+	XRAPI_ATTR XrResult XRAPI_CALL xrEnumerateSwapchainFormats(
+		XrSession										session,
+		uint32_t										formatCapacityInput,
+		uint32_t*										formatCountOutput,
+		int64_t*										formats)
+	{
+		OXRVISIONOS_CHECK_SESSION_EARLY_RETURN(session);
+
+		FOXRVisionOSSession* Session = (FOXRVisionOSSession*)session;
+		return Session->XrEnumerateSwapchainFormats(formatCapacityInput,
+													formatCountOutput,
+													formats);
+	}
+
+	XRAPI_ATTR XrResult XRAPI_CALL xrCreateSwapchain(
+		XrSession										session,
+		const XrSwapchainCreateInfo*					createInfo,
+		XrSwapchain*									swapchain)
+	{
+		OXRVISIONOS_CHECK_SESSION_EARLY_RETURN(session);
+
+		FOXRVisionOSSession* Session = (FOXRVisionOSSession*)session;
+		return Session->XrCreateSwapchain(createInfo,
+													swapchain);
+	}
+
+	XRAPI_ATTR XrResult XRAPI_CALL xrDestroySwapchain(
+		XrSwapchain										swapchain)
+	{
+		OXRVISIONOS_CHECK_SWAPCHAIN_EARLY_RETURN(swapchain);
+		FOXRVisionOSSwapchain* Swapchain = (FOXRVisionOSSwapchain*)swapchain;
+		return Swapchain->XrDestroySwapchain();
+	}
+
+	XRAPI_ATTR XrResult XRAPI_CALL xrEnumerateSwapchainImages(
+		XrSwapchain                                 swapchain,
+		uint32_t                                    imageCapacityInput,
+		uint32_t*									imageCountOutput,
+		XrSwapchainImageBaseHeader*					images)
+	{
+		OXRVISIONOS_CHECK_SWAPCHAIN_EARLY_RETURN(swapchain);
+		FOXRVisionOSSwapchain* Swapchain = (FOXRVisionOSSwapchain*)swapchain;
+		return Swapchain->XrEnumerateSwapchainImages(
+			imageCapacityInput,
+			imageCountOutput,
+			images);
+	}
+
+	XRAPI_ATTR XrResult XRAPI_CALL xrAcquireSwapchainImage(
+		XrSwapchain                                 swapchain,
+		const XrSwapchainImageAcquireInfo*			acquireInfo,
+		uint32_t*									index)
+	{
+		OXRVISIONOS_CHECK_SWAPCHAIN_EARLY_RETURN(swapchain);
+		FOXRVisionOSSwapchain* Swapchain = (FOXRVisionOSSwapchain*)swapchain;
+		return Swapchain->XrAcquireSwapchainImage(
+			acquireInfo,
+			index);
+	}
+
+	XRAPI_ATTR XrResult XRAPI_CALL xrWaitSwapchainImage(
+		XrSwapchain                                 swapchain,
+		const XrSwapchainImageWaitInfo*				waitInfo)
+	{
+		OXRVISIONOS_CHECK_SWAPCHAIN_EARLY_RETURN(swapchain);
+		FOXRVisionOSSwapchain* Swapchain = (FOXRVisionOSSwapchain*)swapchain;
+		return Swapchain->XrWaitSwapchainImage(waitInfo);
+	}
+
+	XRAPI_ATTR XrResult XRAPI_CALL xrReleaseSwapchainImage(
+		XrSwapchain                                 swapchain,
+		const XrSwapchainImageReleaseInfo*			releaseInfo)
+	{
+		OXRVISIONOS_CHECK_SWAPCHAIN_EARLY_RETURN(swapchain);
+		FOXRVisionOSSwapchain* Swapchain = (FOXRVisionOSSwapchain*)swapchain;
+		return Swapchain->XrReleaseSwapchainImage(releaseInfo);
+	}
+
+	XRAPI_ATTR XrResult XRAPI_CALL xrBeginSession(
+		XrSession                                   session,
+		const XrSessionBeginInfo*					beginInfo)
+	{
+		OXRVISIONOS_CHECK_SESSION_EARLY_RETURN(session);
+
+		FOXRVisionOSSession* Session = (FOXRVisionOSSession*)session;
+		XrResult Result = Session->XrBeginSession(beginInfo);
+		OXRVISIONOS_LOG_RESULT(Result);
+		return Result;
+	}
+
+	XRAPI_ATTR XrResult XRAPI_CALL xrEndSession(
+		XrSession                                   session)
+	{
+		OXRVISIONOS_CHECK_SESSION_EARLY_RETURN(session);
+
+		FOXRVisionOSSession* Session = (FOXRVisionOSSession*)session;
+		XrResult Result = Session->XrEndSession();
+		OXRVISIONOS_LOG_RESULT(Result);
+		return Result;
+	}
+
+	XRAPI_ATTR XrResult XRAPI_CALL xrRequestExitSession(
+		XrSession                                   session)
+	{
+		OXRVISIONOS_CHECK_SESSION_EARLY_RETURN(session);
+
+		FOXRVisionOSSession* Session = (FOXRVisionOSSession*)session;
+		XrResult Result = Session->XrRequestExitSession();
+		OXRVISIONOS_LOG_RESULT(Result);
+		return Result;
+	}
+
+	XRAPI_ATTR XrResult XRAPI_CALL xrWaitFrame(
+		XrSession                                   session,
+		const XrFrameWaitInfo*						frameWaitInfo,
+		XrFrameState*								frameState)
+	{
+		OXRVISIONOS_CHECK_SESSION_EARLY_RETURN(session);
+
+		FOXRVisionOSSession* Session = (FOXRVisionOSSession*)session;
+		return Session->XrWaitFrame(frameWaitInfo, frameState);
+	}
+
+	XRAPI_ATTR XrResult XRAPI_CALL xrBeginFrame(
+		XrSession                                   session,
+		const XrFrameBeginInfo*						frameBeginInfo)
+	{
+		OXRVISIONOS_CHECK_SESSION_EARLY_RETURN(session);
+
+		FOXRVisionOSSession* Session = (FOXRVisionOSSession*)session;
+		return Session->XrBeginFrame(frameBeginInfo);
+	}
+
+	XRAPI_ATTR XrResult XRAPI_CALL xrEndFrame(
+		XrSession                                   session,
+		const XrFrameEndInfo*						frameEndInfo)
+	{
+		OXRVISIONOS_CHECK_SESSION_EARLY_RETURN(session);
+
+		FOXRVisionOSSession* Session = (FOXRVisionOSSession*)session;
+		return Session->XrEndFrame(frameEndInfo);
+	}
+
+	XRAPI_ATTR XrResult XRAPI_CALL xrLocateViews(
+		XrSession                                   session,
+		const XrViewLocateInfo*						viewLocateInfo,
+		XrViewState*								viewState,
+		uint32_t                                    viewCapacityInput,
+		uint32_t*									viewCountOutput,
+		XrView*										views)
+	{
+		//OXRVISIONOS_CHECK_SESSION_EARLY_RETURN(session);
+
+		FOXRVisionOSSession* Session = (FOXRVisionOSSession*)session;
+		return Session->XrLocateViews(viewLocateInfo,
+			viewState,
+			viewCapacityInput,
+			viewCountOutput,
+			views);
+	}
+
+	XRAPI_ATTR XrResult XRAPI_CALL xrStringToPath(
+		XrInstance                                  instance,
+		const char*									pathString,
+		XrPath*										path)
+	{
+		OXRVISIONOS_CHECK_INSTANCE_EARLY_RETURN(instance);
+		FOXRVisionOSInstance* Instance = (FOXRVisionOSInstance*)instance;
+		return Instance->XrStringToPath(
+			pathString,
+			path);
+	}
+
+	XRAPI_ATTR XrResult XRAPI_CALL xrPathToString(
+		XrInstance                                  instance,
+		XrPath                                      path,
+		uint32_t                                    bufferCapacityInput,
+		uint32_t*									bufferCountOutput,
+		char*										buffer)
+	{
+		OXRVISIONOS_CHECK_INSTANCE_EARLY_RETURN(instance);
+		FOXRVisionOSInstance* Instance = (FOXRVisionOSInstance*)instance;
+		return Instance->XrPathToString(
+			path,
+			bufferCapacityInput,
+			bufferCountOutput,
+			buffer);
+	}
+
+	XRAPI_ATTR XrResult XRAPI_CALL xrCreateActionSet(
+		XrInstance                                  instance,
+		const XrActionSetCreateInfo*				createInfo,
+		XrActionSet*								actionSet)
+	{
+		OXRVISIONOS_CHECK_INSTANCE_EARLY_RETURN(instance);
+		FOXRVisionOSInstance* Instance = (FOXRVisionOSInstance*)instance;
+		return Instance->XrCreateActionSet(createInfo,
+											actionSet);
+	}
+
+	XRAPI_ATTR XrResult XRAPI_CALL xrDestroyActionSet(
+		XrActionSet                                 actionSet)
+	{
+		FOXRVisionOSActionSet* ActionSet = (FOXRVisionOSActionSet*)actionSet;
+		return ActionSet->XrDestroyActionSet();
+	}
+
+	XRAPI_ATTR XrResult XRAPI_CALL xrCreateAction(
+		XrActionSet                                 actionSet,
+		const XrActionCreateInfo*					createInfo,
+		XrAction*									action)
+	{
+		FOXRVisionOSActionSet* ActionSet = (FOXRVisionOSActionSet*)actionSet;
+		return ActionSet->XrCreateAction(createInfo,
+										action);
+	}
+
+	XRAPI_ATTR XrResult XRAPI_CALL xrDestroyAction(
+		XrAction                                    action)
+	{
+		FOXRVisionOSAction* Action = (FOXRVisionOSAction*)action;
+		return Action->XrDestroyAction();
+	}
+
+	XRAPI_ATTR XrResult XRAPI_CALL xrSuggestInteractionProfileBindings(
+		XrInstance                                  instance,
+		const XrInteractionProfileSuggestedBinding* suggestedBindings)
+	{
+		OXRVISIONOS_CHECK_INSTANCE_EARLY_RETURN(instance);
+		FOXRVisionOSInstance* Instance = (FOXRVisionOSInstance*)instance;
+		return Instance->XrSuggestInteractionProfileBindings(suggestedBindings);
+	}
+
+	XRAPI_ATTR XrResult XRAPI_CALL xrAttachSessionActionSets(
+		XrSession                                   session,
+		const XrSessionActionSetsAttachInfo*		attachInfo)
+	{
+		OXRVISIONOS_CHECK_SESSION_EARLY_RETURN(session);
+
+		FOXRVisionOSSession* Session = (FOXRVisionOSSession*)session;
+		return Session->XrAttachSessionActionSets(attachInfo);
+	}
+
+	XRAPI_ATTR XrResult XRAPI_CALL xrGetCurrentInteractionProfile(
+		XrSession                                   session,
+		XrPath                                      topLevelUserPath,
+		XrInteractionProfileState*					interactionProfile)
+	{
+		FOXRVisionOSSession* Session = (FOXRVisionOSSession*)session;
+		return Session->XrGetCurrentInteractionProfile(topLevelUserPath, interactionProfile);
+	}
+
+	XRAPI_ATTR XrResult XRAPI_CALL xrGetActionStateBoolean(
+		XrSession                                   session,
+		const XrActionStateGetInfo*					getInfo,
+		XrActionStateBoolean*						state)
+	{
+		FOXRVisionOSSession* Session = (FOXRVisionOSSession*)session;
+		return Session->XrGetActionStateBoolean(getInfo, state);
+	}
+
+	XRAPI_ATTR XrResult XRAPI_CALL xrGetActionStateFloat(
+		XrSession                                   session,
+		const XrActionStateGetInfo*					getInfo,
+		XrActionStateFloat*							state)
+	{
+		FOXRVisionOSSession* Session = (FOXRVisionOSSession*)session;
+		return Session->XrGetActionStateFloat(getInfo, state);
+	}
+
+	XRAPI_ATTR XrResult XRAPI_CALL xrGetActionStateVector2f(
+		XrSession                                   session,
+		const XrActionStateGetInfo*					getInfo,
+		XrActionStateVector2f*						state)
+	{
+		FOXRVisionOSSession* Session = (FOXRVisionOSSession*)session;
+		return Session->XrGetActionStateVector2f(getInfo, state);
+	}
+
+	XRAPI_ATTR XrResult XRAPI_CALL xrGetActionStatePose(
+		XrSession                                   session,
+		const XrActionStateGetInfo*					getInfo,
+		XrActionStatePose*							state)
+	{
+		FOXRVisionOSSession* Session = (FOXRVisionOSSession*)session;
+		return Session->XrGetActionStatePose(getInfo, state);
+	}
+
+	XRAPI_ATTR XrResult XRAPI_CALL xrSyncActions(
+		XrSession                                   session,
+		const										XrActionsSyncInfo* syncInfo)
+	{
+		OXRVISIONOS_CHECK_SESSION_EARLY_RETURN(session);
+
+		FOXRVisionOSSession* Session = (FOXRVisionOSSession*)session;
+		return Session->XrSyncActions(syncInfo);
+	}
+
+	XRAPI_ATTR XrResult XRAPI_CALL xrEnumerateBoundSourcesForAction(
+		XrSession                                   session,
+		const XrBoundSourcesForActionEnumerateInfo* enumerateInfo,
+		uint32_t                                    sourceCapacityInput,
+		uint32_t*									sourceCountOutput,
+		XrPath*										sources)
+	{
+		FOXRVisionOSSession* Session = (FOXRVisionOSSession*)session;
+		return Session->XrEnumerateBoundSourcesForAction(
+			enumerateInfo,
+			sourceCapacityInput,
+			sourceCountOutput,
+			sources);
+	}
+
+	XRAPI_ATTR XrResult XRAPI_CALL xrGetInputSourceLocalizedName(
+		XrSession                                   session,
+		const XrInputSourceLocalizedNameGetInfo*	getInfo,
+		uint32_t                                    bufferCapacityInput,
+		uint32_t*									bufferCountOutput,
+		char*										buffer)
+	{
+		// Note: we don't use this so I'm not implementing it.
+		return XrResult::XR_ERROR_FUNCTION_UNSUPPORTED;
+	}
+
+	XRAPI_ATTR XrResult XRAPI_CALL xrApplyHapticFeedback(
+		XrSession                                   session,
+		const XrHapticActionInfo*					hapticActionInfo,
+		const XrHapticBaseHeader*					hapticFeedback)
+	{
+		OXRVISIONOS_CHECK_SESSION_EARLY_RETURN(session);
+
+		FOXRVisionOSSession* Session = (FOXRVisionOSSession*)session;
+		return Session->XrApplyHapticFeedback(
+			hapticActionInfo,
+			hapticFeedback);
+	}
+
+	XRAPI_ATTR XrResult XRAPI_CALL xrStopHapticFeedback(
+		XrSession                                   session,
+		const XrHapticActionInfo*					hapticActionInfo)
+	{
+		OXRVISIONOS_CHECK_SESSION_EARLY_RETURN(session);
+		
+		FOXRVisionOSSession* Session = (FOXRVisionOSSession*)session;
+		return Session->XrStopHapticFeedback(
+			hapticActionInfo);
+	}
+
+// Any supported extension functions would need to be mapped as well.  That could go here.
+
+#define OXRVISIONOS_MAP_GLOBAL_FUNCTION(Type,Func) OXRVisionOS::GlobalFunctionMap.Add(FString(TEXT(#Func)), (PFN_xrVoidFunction)(OXRVisionOS::Func));
+#define OXRVISIONOS_MAP_INSTANCE_FUNCTION(Type,Func) OXRVisionOS::InstanceFunctionMap.Add(FString(TEXT(#Func)), (PFN_xrVoidFunction)(OXRVisionOS::Func));
+
+	void BuildFunctionMaps()
+	{
+		ENUM_XR_ENTRYPOINTS_GLOBAL(OXRVISIONOS_MAP_GLOBAL_FUNCTION);
+		ENUM_XR_ENTRYPOINTS(OXRVISIONOS_MAP_INSTANCE_FUNCTION);
+	}
+
+#undef OXRVISIONOS_MAP_INSTANCE_FUNCTION
+#undef OXRVISIONOS_MAP_GLOBAL_FUNCTION
+
+	void ClearFunctionMaps()
+	{
+		GlobalFunctionMap.Reset();
+		InstanceFunctionMap.Reset();
+	}
+}
+
+
+IMPLEMENT_MODULE(FOXRVisionOS, OXRVisionOS)
+
+void FOXRVisionOS::StartupModule()
+{
+	RegisterOpenXRExtensionModularFeature();
+	IModularFeatures::Get().RegisterModularFeature(GetFeatureName(), this);
+
+	OXRVisionOS::BuildFunctionMaps();
+
+	SupportedExtensions.Add(XR_EPIC_OXRVISIONOS_CONTROLLER_NAME, XR_EPIC_oxrvisionos_controller_SPEC_VERSION);
+//	SupportedExtensions.Add(XR_EXT_EYE_GAZE_INTERACTION_EXTENSION_NAME, XR_EXT_eye_gaze_interaction_SPEC_VERSION);
+}
+
+void FOXRVisionOS::ShutdownModule()
+{
+	OXRVisionOS::ClearFunctionMaps();
+	Instance = nullptr;
+}
+
+bool FOXRVisionOS::GetCustomLoader(PFN_xrGetInstanceProcAddr* OutGetProcAddr)
+{
+	*OutGetProcAddr = &OXRVisionOS::xrGetInstanceProcAddr;
+	return true;
+}
+
+FOpenXRRenderBridge* FOXRVisionOS::GetCustomRenderBridge(XrInstance InInstance)
+{
+	RenderBridge = CreateRenderBridge_OXRVisionOS(InInstance);
+	return RenderBridge;
+}
+
+void FOXRVisionOS::OnBeginRendering_GameThread(XrSession InSession)
+{
+    if (!FOXRVisionOSInstance::CheckSession(InSession))
+    {
+        return;
+    }
+    FOXRVisionOSSession* Session = (FOXRVisionOSSession*)InSession;
+    Session->OnBeginRendering_GameThread();
+}
+
+void FOXRVisionOS::OnBeginRendering_RenderThread(XrSession InSession)
+{
+	if (!FOXRVisionOSInstance::CheckSession(InSession))
+	{
+		return;
+	}
+	FOXRVisionOSSession* Session = (FOXRVisionOSSession*)InSession;
+	Session->OnBeginRendering_RenderThread();
+}
+
+FOXRVisionOS* FOXRVisionOS::Get()
+{
+	TArray<FOXRVisionOS*> Impls = IModularFeatures::Get().GetModularFeatureImplementations<FOXRVisionOS>(GetFeatureName());
+
+	// There can be only one!  Or zero.  The implementations are platform specific and we are not currently supporting 'overlapping' platforms.
+	check(Impls.Num() <= 1);
+
+	if (Impls.Num() > 0)
+	{
+		check(Impls[0]);
+		return Impls[0];
+	}
+	return nullptr;
+}
+
+XrResult FOXRVisionOS::XrEnumerateApiLayerProperties(
+	uint32_t  propertyCapacityInput,
+	uint32_t* propertyCountOutput,
+	XrApiLayerProperties* properties)
+{
+	if (properties == nullptr)
+	{
+		*propertyCountOutput = 0; //TODO make real if we need to support any layers
+		return XrResult::XR_SUCCESS;
+	}
+	else
+	{
+		return XrResult::XR_ERROR_RUNTIME_FAILURE;
+	}
+}
+
+bool FOXRVisionOS::IsExtensionSupported(const ANSICHAR* InName) const
+{
+	return SupportedExtensions.Contains(InName);
+}
+void FOXRVisionOS::PostSyncActions(XrSession InSession)
+{
+	if (!FOXRVisionOSInstance::CheckSession(InSession))
+	{
+		return;
+	}
+	FOXRVisionOSSession* Session = (FOXRVisionOSSession*)InSession;
+	Session->OnPostSyncActions();
+}
+
+XrResult FOXRVisionOS::XrEnumerateInstanceExtensionProperties(
+	const char* layerName,
+	uint32_t propertyCapacityInput,
+	uint32_t* propertyCountOutput,
+	XrExtensionProperties* properties)
+{
+	if (propertyCountOutput == nullptr)
+	{
+		return XrResult::XR_ERROR_VALIDATION_FAILURE;
+	}
+	if (propertyCapacityInput != 0 && properties == nullptr)
+	{
+		return XrResult::XR_ERROR_VALIDATION_FAILURE;
+	}
+
+	*propertyCountOutput = SupportedExtensions.Num();
+	if (propertyCapacityInput != 0)
+	{
+		if (propertyCapacityInput < *propertyCountOutput)
+		{
+			return XR_ERROR_SIZE_INSUFFICIENT;
+		}
+
+		int Index = 0;
+		for (const auto& Pair : SupportedExtensions)
+		{
+			FCStringAnsi::Strncpy(properties[Index].extensionName, Pair.Key, XR_MAX_EXTENSION_NAME_SIZE);
+			properties[Index].extensionVersion = Pair.Value;
+			++Index;
+		}
+	}
+	return XrResult::XR_SUCCESS;
+}
+
+XrResult FOXRVisionOS::XrCreateInstance(
+	const XrInstanceCreateInfo* createInfo,
+	XrInstance* instance)
+{
+	// Only one instance supported
+	if (Instance.IsValid())
+	{
+		return XrResult::XR_ERROR_LIMIT_REACHED;
+	}
+
+	XrResult Ret = FOXRVisionOSInstance::Create(Instance, createInfo, this);
+	*instance = (XrInstance)Instance.Get();
+	return Ret;
+
+}
+
+XrResult FOXRVisionOS::XrDestroyInstance(
+	XrInstance instance)
+{
+	OXRVISIONOS_CHECK_INSTANCE_EARLY_RETURN(instance);
+
+	Instance = nullptr;
+	return XR_SUCCESS;
+}
diff --git a/Engine/Platforms/VisionOS/Plugins/Runtime/OpenXRVisionOS/Source/OXRVisionOS/Private/OXRVisionOS.h b/Engine/Platforms/VisionOS/Plugins/Runtime/OpenXRVisionOS/Source/OXRVisionOS/Private/OXRVisionOS.h
new file mode 100644
--- /dev/null
+++ b/Engine/Platforms/VisionOS/Plugins/Runtime/OpenXRVisionOS/Source/OXRVisionOS/Private/OXRVisionOS.h
@@ -0,0 +1,114 @@
+// Copyright Epic Games, Inc. All Rights Reserved.
+
+#pragma once
+
+#include "CoreMinimal.h"
+
+#include "Modules/ModuleInterface.h"
+#include "IOpenXRExtensionPlugin.h"
+
+#include "Logging/LogMacros.h"
+
+#include "OpenXRCore.h"
+#include "Containers/Map.h"
+
+// Uncomment to make intellisense work better.
+//#include "../../../../../../../../Source/ThirdParty/OpenXR/include/openxr/openxr.h"
+
+DECLARE_LOG_CATEGORY_EXTERN(LogOXRVisionOS, Verbose, All);
+// Note: if you are turning this log up to verbose you may also want LogMetalVisionOS turned up to verbose.
+
+class FOXRVisionOSInstance;
+
+class FOXRVisionOS : public IModuleInterface, public IOpenXRExtensionPlugin
+{
+public:
+	/************************************************************************/
+	/* IModuleInterface                                                     */
+	/************************************************************************/
+	virtual void StartupModule() override;
+	virtual void ShutdownModule() override;
+
+	/************************************************************************/
+	/* IOpenXRExtensionPlugin                                               */
+	/************************************************************************/
+	virtual FString GetDisplayName() override
+	{
+		return FString(TEXT("OXRVisionOS"));
+	}
+	virtual bool GetCustomLoader(PFN_xrGetInstanceProcAddr* OutGetProcAddr) override;
+	virtual bool IsStandaloneStereoOnlyDevice() override { return true; }
+	virtual class FOpenXRRenderBridge* GetCustomRenderBridge(XrInstance InInstance) override;
+    virtual void OnBeginRendering_GameThread(XrSession InSession) override;
+	virtual void OnBeginRendering_RenderThread(XrSession InSession) override;
+
+	/************************************************************************/
+	/* FOXRVisionOS                                                             */
+	/************************************************************************/
+	static FOXRVisionOS* Get();
+	static const FOXRVisionOSInstance* GetInstance() { return Instance.Get(); }
+
+
+	XrResult XrEnumerateApiLayerProperties(
+		uint32_t  propertyCapacityInput,
+		uint32_t* propertyCountOutput,
+		XrApiLayerProperties* properties);
+
+	XrResult XrEnumerateInstanceExtensionProperties(
+		const char* layerName,
+		uint32_t propertyCapacityInput,
+		uint32_t* propertyCountOutput,
+		XrExtensionProperties* properties);
+
+	XrResult XrCreateInstance(
+		const XrInstanceCreateInfo* createInfo,
+		XrInstance* instance);
+
+	XrResult XrDestroyInstance(
+		XrInstance instance);
+
+	static bool CheckInstance(XrInstance instance) { return (instance != XR_NULL_HANDLE) && (instance == (XrInstance)Instance.Get()); }
+
+	FOpenXRRenderBridge* GetRenderBridge() const
+	{
+		return RenderBridge;
+	}
+
+	bool IsExtensionSupported(const ANSICHAR* InName) const;
+	virtual void PostSyncActions(XrSession InSession) override;
+
+private:
+	static FName GetFeatureName()
+	{
+		static FName FeatureName = FName(TEXT("OXRVisionOS"));
+		return FeatureName;
+	}
+
+private:
+	struct FStringToVersionKeyFunc
+	{
+		typedef const ANSICHAR* KeyType;
+		typedef const ANSICHAR* KeyInitType;
+		typedef const TPairInitializer<const ANSICHAR*, int32_t>& ElementInitType;
+
+		enum { bAllowDuplicateKeys = false };
+
+		static FORCEINLINE bool Matches(const ANSICHAR* A, const ANSICHAR* B)
+		{
+			return TCString<ANSICHAR>::Stricmp(A, B) == 0;
+		}
+
+		static FORCEINLINE uint32 GetKeyHash(const ANSICHAR* Key)
+		{
+			return FCrc::Strihash_DEPRECATED(Key);
+		}
+
+		static FORCEINLINE KeyInitType GetSetKey(ElementInitType Element)
+		{
+			return Element.Key;
+		}
+	};	
+	TMap<const ANSICHAR*, int32_t, FDefaultSetAllocator, FStringToVersionKeyFunc> SupportedExtensions;
+	static TSharedPtr<FOXRVisionOSInstance, ESPMode::ThreadSafe> Instance;
+	FOpenXRRenderBridge* RenderBridge = nullptr;
+};
diff --git a/Engine/Platforms/VisionOS/Plugins/Runtime/OpenXRVisionOS/Source/OXRVisionOS/Private/OXRVisionOSAction.cpp b/Engine/Platforms/VisionOS/Plugins/Runtime/OpenXRVisionOS/Source/OXRVisionOS/Private/OXRVisionOSAction.cpp
new file mode 100644
--- /dev/null
+++ b/Engine/Platforms/VisionOS/Plugins/Runtime/OpenXRVisionOS/Source/OXRVisionOS/Private/OXRVisionOSAction.cpp
@@ -0,0 +1,390 @@
+// Copyright Epic Games, Inc. All Rights Reserved.
+
+#include "OXRVisionOSAction.h"
+#include "OXRVisionOSActionSet.h"
+#include "OXRVisionOSInstance.h"
+#include "OXRVisionOSController.h"
+#include "OXRVisionOSTracker.h"
+#include "HAL/PlatformMemory.h"
+
+XrResult FOXRVisionOSAction::Create(TSharedPtr<FOXRVisionOSAction, ESPMode::ThreadSafe>& OutAction, const XrActionCreateInfo* createInfo, FOXRVisionOSActionSet* ActionSet)
+{
+	if (createInfo == nullptr || ActionSet == nullptr)
+	{
+		return XrResult::XR_ERROR_VALIDATION_FAILURE;
+	}
+
+	if (createInfo->type != XR_TYPE_ACTION_CREATE_INFO)
+	{
+		return XrResult::XR_ERROR_VALIDATION_FAILURE;
+	}
+
+	OutAction = MakeShared<FOXRVisionOSAction, ESPMode::ThreadSafe>(createInfo, ActionSet);
+	if (OutAction->bCreateFailed)
+	{
+		OutAction = nullptr;
+		return XrResult::XR_ERROR_RUNTIME_FAILURE;
+	}
+
+	return XrResult::XR_SUCCESS;
+}
+
+FOXRVisionOSAction::FOXRVisionOSAction(const XrActionCreateInfo* createInfo, FOXRVisionOSActionSet* InActionSet)
+{
+	ActionSet = InActionSet;
+
+	check(createInfo->next == nullptr);
+	ActionName = FName(createInfo->actionName);
+	ActionType = createInfo->actionType;
+	SubactionPathBoundSourcesAndStates.Add(XR_NULL_PATH);
+	for (int i = 0; i < createInfo->countSubactionPaths; i++)
+	{
+		XrPath SubActionPath = createInfo->subactionPaths[i];
+		SubActionPathBinding& NewBinding = SubactionPathBoundSourcesAndStates.Add(SubActionPath);
+		FMemory::Memzero(&NewBinding.ActionState, sizeof(ActionStateUnion));
+	}
+	FCStringAnsi::Strncpy(LocalizedActionName, createInfo->localizedActionName, XR_MAX_LOCALIZED_ACTION_NAME_SIZE);
+}
+
+FOXRVisionOSAction::~FOXRVisionOSAction()
+{
+	if (bCreateFailed)
+	{
+		UE_LOG(LogOXRVisionOS, Warning, TEXT("Destructing FOXRVisionOSAction because create failed."));
+	}
+}
+
+XrResult FOXRVisionOSAction::XrDestroyAction()
+{
+	// ActionSet->DestroyAction() can delete this, so better just return after that.
+	return ActionSet->DestroyAction(this);
+}
+
+FOXRVisionOSAction::EActionStateTranslation FOXRVisionOSAction::GetActionStateTranslation(EOXRVisionOSControllerButton Button) const
+{
+	if (ActionType == XR_ACTION_TYPE_BOOLEAN_INPUT && FOXRVisionOSController::IsFloatButton(Button))
+	{
+		return FOXRVisionOSAction::EActionStateTranslation::FloatAsBool;
+	}
+
+	return FOXRVisionOSAction::EActionStateTranslation::None;
+}
+
+void FOXRVisionOSAction::BindSource(XrPath AppPath, EOXRVisionOSControllerButton Button)
+{
+	BoundPaths.AddUnique(AppPath);
+
+	char* PathString = ActionSet->GetInstance()->PathToString(AppPath);
+	for (auto& Elem : SubactionPathBoundSourcesAndStates)
+	{
+		if (Elem.Key == XR_NULL_PATH)
+		{
+			// Null path queries all paths.
+			Elem.Value.Buttons.Emplace(Button, GetActionStateTranslation(Button));
+		}
+		else
+		{
+			// If AppPath starts with any of the subaction paths for this action the binding also goes for that subaction path
+			char* SubactionPathString = ActionSet->GetInstance()->PathToString(Elem.Key);
+			int32 Len = FCStringAnsi::Strlen(SubactionPathString);
+			if (0 == FCStringAnsi::Strncmp(SubactionPathString, PathString, Len))
+			{
+				Elem.Value.Buttons.Emplace(Button, GetActionStateTranslation(Button));
+			}
+		}
+	}
+}
+
+const TArray<XrPath>& FOXRVisionOSAction::GetBoundSources() const
+{
+	return BoundPaths;
+}
+
+XrResult FOXRVisionOSAction::GetActionStateBoolean(XrPath SubActionPath, XrActionStateBoolean& OutState) const
+{
+	if (ActionType != XR_ACTION_TYPE_BOOLEAN_INPUT)
+	{
+		return XrResult::XR_ERROR_ACTION_TYPE_MISMATCH;
+	}
+
+	const SubActionPathBinding* Binding = SubactionPathBoundSourcesAndStates.Find(SubActionPath);
+	if (Binding == nullptr)
+	{
+		return XrResult::XR_ERROR_PATH_UNSUPPORTED;
+	}
+
+	OutState = Binding->ActionState.Boolean;
+
+	return XrResult::XR_SUCCESS;
+}
+
+XrResult FOXRVisionOSAction::GetActionStateFloat(XrPath SubActionPath, XrActionStateFloat& OutState) const
+{
+	if (ActionType != XR_ACTION_TYPE_FLOAT_INPUT)
+	{
+		return XrResult::XR_ERROR_ACTION_TYPE_MISMATCH;
+	}
+
+	const SubActionPathBinding* Binding = SubactionPathBoundSourcesAndStates.Find(SubActionPath);
+	if (Binding == nullptr)
+	{
+		return XrResult::XR_ERROR_PATH_UNSUPPORTED;
+	}
+
+	OutState = Binding->ActionState.Float;
+
+	return XrResult::XR_SUCCESS;
+}
+
+XrResult FOXRVisionOSAction::GetActionStateVector2f(XrPath SubActionPath, XrActionStateVector2f& OutState) const
+{
+	if (ActionType != XR_ACTION_TYPE_VECTOR2F_INPUT)
+	{
+		return XrResult::XR_ERROR_ACTION_TYPE_MISMATCH;
+	}
+
+	const SubActionPathBinding* Binding = SubactionPathBoundSourcesAndStates.Find(SubActionPath);
+	if (Binding == nullptr)
+	{
+		return XrResult::XR_ERROR_PATH_UNSUPPORTED;
+	}
+
+	OutState = Binding->ActionState.Vector2f;
+
+	return XrResult::XR_SUCCESS;
+}
+
+XrResult FOXRVisionOSAction::GetActionStatePose(XrPath SubActionPath, XrActionStatePose& OutState) const
+{
+	if (ActionType != XR_ACTION_TYPE_POSE_INPUT)
+	{
+		return XrResult::XR_ERROR_ACTION_TYPE_MISMATCH;
+	}
+
+	const SubActionPathBinding* Binding = SubactionPathBoundSourcesAndStates.Find(SubActionPath);
+	if (Binding == nullptr)
+	{
+		return XrResult::XR_ERROR_PATH_UNSUPPORTED;
+	}
+
+	OutState = Binding->ActionState.Pose;
+
+	return XrResult::XR_SUCCESS;
+}
+
+void FOXRVisionOSAction::Sync(TSharedPtr<FOXRVisionOSController, ESPMode::ThreadSafe>& Controllers, TSharedPtr <FOXRVisionOSTracker, ESPMode::ThreadSafe>& Tracker)
+{
+	for (auto& Elem : SubactionPathBoundSourcesAndStates)
+	{
+		ActionStateUnion & State = Elem.Value.ActionState;
+		TArray<TPair<EOXRVisionOSControllerButton, EActionStateTranslation>>& Buttons = Elem.Value.Buttons;
+
+		switch (ActionType)
+		{
+		case XR_ACTION_TYPE_BOOLEAN_INPUT:
+			SyncPathsBoolean(Controllers, Buttons, State.Boolean);
+			break;
+		case XR_ACTION_TYPE_FLOAT_INPUT:
+			SyncPathsFloat(Controllers, Buttons, State.Float);
+			break;
+		case XR_ACTION_TYPE_VECTOR2F_INPUT:
+			SyncPathsVector2f(Controllers, Buttons, State.Vector2f);
+			break;
+		case XR_ACTION_TYPE_POSE_INPUT:
+			SyncPathsPose(Tracker, Buttons, State.Pose);
+			break;
+		case XR_ACTION_TYPE_VIBRATION_OUTPUT:
+			break;
+		default:
+			check(false);
+			break;
+		}
+	}
+}
+
+void FOXRVisionOSAction::SyncPathsBoolean(TSharedPtr<FOXRVisionOSController, ESPMode::ThreadSafe>& Controllers, const TArray<TPair<EOXRVisionOSControllerButton, EActionStateTranslation>>& Buttons, XrActionStateBoolean& OutActionState)
+{
+	// OpenXR spec says boolean input sources are combined as: OR
+
+	bool bNewValue = false;
+	bool bNewActive = false;
+	XrTime NewTime = 0;
+	for (const TPair<EOXRVisionOSControllerButton, EActionStateTranslation>& Button : Buttons)
+	{
+		bool bValue = false;
+		bool bActive = false;
+		XrTime Time = 0;
+		switch (Button.Value)
+		{
+		case EActionStateTranslation::None:
+			Controllers->GetActionStateBoolean(Button.Key, bValue, bActive, Time);
+			break;
+		case EActionStateTranslation::FloatAsBool:
+		{
+			float fValue;
+			Controllers->GetActionStateFloat(Button.Key, fValue, bActive, Time);
+			// Flip to false if true and value == 0, flip to true if false and value > 0.9.  We are detecting button release (which already has a deadzone) or nearly full press.
+			bValue = OutActionState.currentState;
+			if (bValue && fValue == 0.0f)
+			{
+				bValue = false;
+			}
+			else if(!bValue && fValue > 0.9f)
+			{
+				bValue = true;
+			}
+		}
+			break;
+		default:
+			check(false);
+			break;
+		}
+		if (bActive)
+		{
+			bNewActive = true;
+			bNewValue |= bValue;
+			NewTime = Time;
+		}
+	}
+	OutActionState.isActive = bNewActive;
+	if (OutActionState.isActive)
+	{
+		OutActionState.changedSinceLastSync = OutActionState.currentState != bNewValue;
+		if (OutActionState.changedSinceLastSync)
+		{
+			OutActionState.lastChangeTime = NewTime;
+		}
+		OutActionState.currentState = bNewValue;
+	}
+
+	//if (OutActionState.changedSinceLastSync)
+	//{
+	//	UE_LOG(LogOXRVisionOS, Log, TEXT("SyncPathsBool %s Value: %i Timestamp: %lld"), ANSI_TO_TCHAR(ActionName), OutActionState.currentState, OutActionState.lastChangeTime);
+	//}
+}
+
+void FOXRVisionOSAction::SyncPathsFloat(TSharedPtr<FOXRVisionOSController, ESPMode::ThreadSafe>& Controllers, const TArray<TPair<EOXRVisionOSControllerButton, EActionStateTranslation>>& Buttons, XrActionStateFloat& OutActionState)
+{
+	// OpenXR spec says float input sources are combined as: largest absolute value
+
+	bool bNoValueYetRead = true;
+	float fNewValue = 0.0f;
+	bool bNewActive = false;
+	XrTime NewTime = 0;
+	for (const TPair<EOXRVisionOSControllerButton, EActionStateTranslation>& Button : Buttons)
+	{
+		float fValue = 0.0f;
+		bool bActive = false;
+		XrTime Time = 0;
+		Controllers->GetActionStateFloat(Button.Key, fValue, bActive, Time);
+		if (bActive)
+		{
+			bNewActive = true;
+			if (bNoValueYetRead || FMath::Abs(fValue) > FMath::Abs(fNewValue))
+			{
+				bNoValueYetRead = false;
+				fNewValue = fValue;
+				NewTime = Time;
+			}
+		}
+	}
+	OutActionState.isActive = bNewActive;
+	if (OutActionState.isActive)
+	{
+		OutActionState.changedSinceLastSync = OutActionState.currentState != fNewValue;
+		if (OutActionState.changedSinceLastSync)
+		{
+			OutActionState.lastChangeTime = NewTime;
+		}
+		OutActionState.currentState = fNewValue;
+	}
+
+	//if (OutActionState.changedSinceLastSync)
+	//{
+	//	UE_LOG(LogOXRVisionOS, Log, TEXT("SyncPathsFloat %s Value: %0.2f Timestamp: %lld"), ANSI_TO_TCHAR(ActionName), OutActionState.currentState, OutActionState.lastChangeTime);
+	//}
+}
+
+void FOXRVisionOSAction::SyncPathsVector2f(TSharedPtr<FOXRVisionOSController, ESPMode::ThreadSafe>& Controllers, const TArray<TPair<EOXRVisionOSControllerButton, EActionStateTranslation>>& Buttons, XrActionStateVector2f& OutActionState)
+{
+	// OpenXR spec says pose input sources are combined as: largest vector magnitude wins
+
+	bool bNoValueYetRead = true;
+	XrVector2f NewValue = { 0.0f, 0.0f };
+	float NewMagnitudeSq = 0.0f;
+	bool bNewActive = false;
+	XrTime NewTime = 0;
+	for (const TPair<EOXRVisionOSControllerButton, EActionStateTranslation>& Button : Buttons)
+	{
+		XrVector2f Value = {0.0f, 0.0f};
+		bool bActive = false;
+		XrTime Time = 0;
+		Controllers->GetActionStateVector2f(Button.Key, Value, bActive, Time);
+		if (bActive)
+		{
+			bNewActive = true;
+			const float MagnitudeSq = (Value.x * Value.x) + (Value.y * Value.y);
+			if (bNoValueYetRead || (MagnitudeSq > NewMagnitudeSq))
+			{
+				bNoValueYetRead = false;
+				NewMagnitudeSq = MagnitudeSq;
+				NewValue = Value;
+				NewTime = Time;
+			}
+		}
+	}
+	OutActionState.isActive = bNewActive;
+	if (OutActionState.isActive)
+	{
+		OutActionState.changedSinceLastSync = (OutActionState.currentState.x != NewValue.x) || (OutActionState.currentState.y != NewValue.y);
+		if (OutActionState.changedSinceLastSync)
+		{
+			OutActionState.lastChangeTime = NewTime;
+		}
+		OutActionState.currentState = NewValue;
+	}
+
+	//if (OutActionState.changedSinceLastSync)
+	//{
+	//	UE_LOG(LogOXRVisionOS, Log, TEXT("SyncPathsVector2f %s Value: %0.2f,%0.2f Timestamp: %lld"), ANSI_TO_TCHAR(ActionName), OutActionState.currentState.x, OutActionState.currentState.y, OutActionState.lastChangeTime);
+	//}
+
+}
+
+void FOXRVisionOSAction::SyncPathsPose(TSharedPtr <FOXRVisionOSTracker, ESPMode::ThreadSafe>& Tracker, const TArray<TPair<EOXRVisionOSControllerButton, EActionStateTranslation>>& Buttons, XrActionStatePose& OutActionState)
+{
+	// OpenXR spec says pose input sources are combined as: always the same source
+
+	bool bNewActive = false;
+	for (const TPair<EOXRVisionOSControllerButton, EActionStateTranslation>& Button : Buttons)
+	{
+		bool bActive = false;
+		Tracker->GetActionStatePose(Button.Key, bActive);
+		if (bActive)
+		{
+			bNewActive = true;
+		}
+	}
+	OutActionState.isActive = bNewActive;
+}
+
+EOXRVisionOSControllerButton FOXRVisionOSAction::GetPoseButton(XrPath SubactionPath) const
+{
+	const SubActionPathBinding* Binding = SubactionPathBoundSourcesAndStates.Find(SubactionPath);
+
+	if (Binding)
+	{
+		if (Binding->Buttons.Num() > 0)
+		{
+			return Binding->Buttons[0].Key;  // just return the first button
+		}
+		else
+		{
+			return EOXRVisionOSControllerButton::NullInput;
+		}
+	}
+	else
+	{
+		return EOXRVisionOSControllerButton::NullInput;
+	}
+}
\ No newline at end of file
diff --git a/Engine/Platforms/VisionOS/Plugins/Runtime/OpenXRVisionOS/Source/OXRVisionOS/Private/OXRVisionOSAction.h b/Engine/Platforms/VisionOS/Plugins/Runtime/OpenXRVisionOS/Source/OXRVisionOS/Private/OXRVisionOSAction.h
new file mode 100644
--- /dev/null
+++ b/Engine/Platforms/VisionOS/Plugins/Runtime/OpenXRVisionOS/Source/OXRVisionOS/Private/OXRVisionOSAction.h
@@ -0,0 +1,73 @@
+// Copyright Epic Games, Inc. All Rights Reserved.
+
+#pragma once
+
+#include "OXRVisionOS.h"
+
+class FOXRVisionOSInstance;
+class FOXRVisionOSActionSet;
+class FOXRVisionOSController;
+class FOXRVisionOSTracker;
+
+enum class EOXRVisionOSControllerButton : int32;
+
+class FOXRVisionOSAction
+{
+public:
+	static XrResult Create(TSharedPtr<FOXRVisionOSAction, ESPMode::ThreadSafe>& OutAction, const XrActionCreateInfo* createInfo, FOXRVisionOSActionSet* ActionSet);
+	FOXRVisionOSAction(const XrActionCreateInfo* createInfo, FOXRVisionOSActionSet* ActionSet);
+	~FOXRVisionOSAction();
+	XrResult XrDestroyAction();
+
+	XrActionType GetActionType() const { return ActionType; }
+	const FOXRVisionOSActionSet& GetActionSet() const { check(ActionSet); return *ActionSet; }
+
+	void BindSource(XrPath AppPath, EOXRVisionOSControllerButton Button);
+	const TArray<XrPath>& GetBoundSources() const;
+
+	XrResult GetActionStateBoolean(XrPath SubActionPath, XrActionStateBoolean& OutState) const;
+	XrResult GetActionStateFloat(XrPath SubActionPath, XrActionStateFloat& OutState) const;
+	XrResult GetActionStateVector2f(XrPath SubActionPath, XrActionStateVector2f& OutState) const;
+	XrResult GetActionStatePose(XrPath SubActionPath, XrActionStatePose& OutState) const;
+
+	void Sync(TSharedPtr<FOXRVisionOSController, ESPMode::ThreadSafe>& Controllers, TSharedPtr <FOXRVisionOSTracker, ESPMode::ThreadSafe>& Tracker);
+	EOXRVisionOSControllerButton GetPoseButton(XrPath SubactionPath) const;
+
+	FName GetActionName() const { return ActionName; }
+
+private:
+	bool bCreateFailed = false;
+	FOXRVisionOSActionSet* ActionSet = nullptr;
+
+	FName					ActionName;
+	XrActionType			ActionType;
+	char					LocalizedActionName[XR_MAX_LOCALIZED_ACTION_NAME_SIZE];
+
+	TArray<XrPath> BoundPaths;
+	// Note we craftily use XR_NULL_PATH as a subaction path to handle the case where none is specified on query.
+	union ActionStateUnion
+	{
+		XrActionStateBoolean Boolean;
+		XrActionStateFloat Float;
+		XrActionStateVector2f Vector2f;
+		XrActionStatePose Pose;
+	};
+	enum class EActionStateTranslation
+	{
+		None,
+		FloatAsBool
+	};
+	struct SubActionPathBinding
+	{
+		TArray<TPair<EOXRVisionOSControllerButton, EActionStateTranslation>> Buttons;
+		ActionStateUnion ActionState = {};
+	};
+	TMap<XrPath, SubActionPathBinding> SubactionPathBoundSourcesAndStates;
+	EActionStateTranslation GetActionStateTranslation(EOXRVisionOSControllerButton Button) const;
+
+
+	void SyncPathsBoolean(TSharedPtr<FOXRVisionOSController, ESPMode::ThreadSafe>& Controllers, const TArray<TPair<EOXRVisionOSControllerButton, EActionStateTranslation>>& Buttons, XrActionStateBoolean& OutActionState);
+	void SyncPathsFloat(TSharedPtr<FOXRVisionOSController, ESPMode::ThreadSafe>& Controllers, const TArray<TPair<EOXRVisionOSControllerButton, EActionStateTranslation>>& Buttons, XrActionStateFloat& OutActionState);
+	void SyncPathsVector2f(TSharedPtr<FOXRVisionOSController, ESPMode::ThreadSafe>& Controllers, const TArray<TPair<EOXRVisionOSControllerButton, EActionStateTranslation>>& Buttons, XrActionStateVector2f& OutActionState);
+	void SyncPathsPose(TSharedPtr <FOXRVisionOSTracker, ESPMode::ThreadSafe>& Tracker, const TArray<TPair<EOXRVisionOSControllerButton, EActionStateTranslation>>& Buttons, XrActionStatePose& OutActionState);
+};
diff --git a/Engine/Platforms/VisionOS/Plugins/Runtime/OpenXRVisionOS/Source/OXRVisionOS/Private/OXRVisionOSActionSet.cpp b/Engine/Platforms/VisionOS/Plugins/Runtime/OpenXRVisionOS/Source/OXRVisionOS/Private/OXRVisionOSActionSet.cpp
new file mode 100644
--- /dev/null
+++ b/Engine/Platforms/VisionOS/Plugins/Runtime/OpenXRVisionOS/Source/OXRVisionOS/Private/OXRVisionOSActionSet.cpp
@@ -0,0 +1,119 @@
+// Copyright Epic Games, Inc. All Rights Reserved.
+
+#include "OXRVisionOSActionSet.h"
+#include "OXRVisionOSInstance.h"
+#include "OXRVisionOSAction.h"
+#include "OXRVisionOSController.h"
+
+XrResult FOXRVisionOSActionSet::Create(TSharedPtr<FOXRVisionOSActionSet, ESPMode::ThreadSafe>& OutActionSet, const XrActionSetCreateInfo* createInfo, FOXRVisionOSInstance* Instance)
+{
+	if (createInfo == nullptr || Instance == nullptr)
+	{
+		return XrResult::XR_ERROR_VALIDATION_FAILURE;
+	}
+
+	if (createInfo->type != XR_TYPE_ACTION_SET_CREATE_INFO)
+	{
+		return XrResult::XR_ERROR_VALIDATION_FAILURE;
+	}
+
+
+	OutActionSet = MakeShared<FOXRVisionOSActionSet, ESPMode::ThreadSafe>(createInfo, Instance);
+	if (OutActionSet->bCreateFailed)
+	{
+		OutActionSet = nullptr;
+		return XrResult::XR_ERROR_RUNTIME_FAILURE;
+	}
+	return XrResult::XR_SUCCESS;
+}
+
+FOXRVisionOSActionSet::FOXRVisionOSActionSet(const XrActionSetCreateInfo* createInfo, FOXRVisionOSInstance* InInstance)
+{
+	Instance = InInstance;
+
+	CreateInfo = *createInfo;
+}
+
+FOXRVisionOSActionSet::~FOXRVisionOSActionSet()
+{
+	if (bCreateFailed)
+	{
+		UE_LOG(LogOXRVisionOS, Warning, TEXT("Destructing FOXRVisionOSActionSet because session create failed."));
+	}
+}
+
+XrResult FOXRVisionOSActionSet::XrDestroyActionSet()
+{
+	// Instance->DestroyActionSet() can delete this, so better just return after that.
+	return Instance->DestroyActionSet(this);
+}
+
+XrResult FOXRVisionOSActionSet::XrCreateAction(
+	const XrActionCreateInfo* createInfo,
+	XrAction* action)
+{
+	if (IsAttached())
+	{
+		return XR_ERROR_ACTIONSETS_ALREADY_ATTACHED;
+	}
+
+	TSharedPtr<FOXRVisionOSAction, ESPMode::ThreadSafe> NewAction;
+	XrResult Ret = FOXRVisionOSAction::Create(NewAction, createInfo, this);
+	if (Ret == XrResult::XR_SUCCESS)
+	{
+		Actions.Add(NewAction);
+		*action = (XrAction)NewAction.Get();
+	}
+	return Ret;
+}
+
+XrResult FOXRVisionOSActionSet::DestroyAction(FOXRVisionOSAction* Action)
+{
+	uint32 ArrayIndex = Actions.IndexOfByPredicate([Action](const TSharedPtr<FOXRVisionOSAction, ESPMode::ThreadSafe>& Data) { return (Data.Get() == Action); });
+
+	if (ArrayIndex == INDEX_NONE)
+	{
+		return  XrResult::XR_ERROR_HANDLE_INVALID;
+	}
+
+	Actions.RemoveAtSwap(ArrayIndex);
+
+	return XrResult::XR_SUCCESS;
+}	
+
+XrResult FOXRVisionOSActionSet::Attach(FOXRVisionOSSession* InSession)
+{
+	Session = InSession;
+
+	return XrResult::XR_SUCCESS;
+}
+
+void FOXRVisionOSActionSet::ClearActive()
+{
+	bAllActive = false;
+}
+
+void FOXRVisionOSActionSet::SetActive(XrPath SubActionPath)
+{
+	bAllActive = true;
+
+	if (SubActionPath != XR_NULL_PATH)
+	{
+		// OXRVisionOS does not currently correctly support action set per-subactionpath activation.
+		check(false);
+	}
+}
+
+void FOXRVisionOSActionSet::SyncActions(TSharedPtr<FOXRVisionOSController, ESPMode::ThreadSafe>& Controllers, TSharedPtr <FOXRVisionOSTracker, ESPMode::ThreadSafe>& Tracker)
+{
+	if (bAllActive == false)
+	{
+		return;
+	}
+
+	for (TSharedPtr<FOXRVisionOSAction, ESPMode::ThreadSafe>& Action : Actions)
+	{
+		// Update the state
+		Action->Sync(Controllers, Tracker);
+	}
+}
diff --git a/Engine/Platforms/VisionOS/Plugins/Runtime/OpenXRVisionOS/Source/OXRVisionOS/Private/OXRVisionOSActionSet.h b/Engine/Platforms/VisionOS/Plugins/Runtime/OpenXRVisionOS/Source/OXRVisionOS/Private/OXRVisionOSActionSet.h
new file mode 100644
--- /dev/null
+++ b/Engine/Platforms/VisionOS/Plugins/Runtime/OpenXRVisionOS/Source/OXRVisionOS/Private/OXRVisionOSActionSet.h
@@ -0,0 +1,44 @@
+// Copyright Epic Games, Inc. All Rights Reserved.
+
+#pragma once
+
+#include "OXRVisionOS.h"
+
+class FOXRVisionOSInstance;
+class FOXRVisionOSSession;
+class FOXRVisionOSAction;
+class FOXRVisionOSController;
+class FOXRVisionOSTracker;
+
+class FOXRVisionOSActionSet
+{
+public:
+	static XrResult Create(TSharedPtr<FOXRVisionOSActionSet, ESPMode::ThreadSafe>& OutActionSet, const XrActionSetCreateInfo* createInfo, FOXRVisionOSInstance* Instance);
+	FOXRVisionOSActionSet(const XrActionSetCreateInfo* createInfo, FOXRVisionOSInstance* Instance);
+	~FOXRVisionOSActionSet();
+	XrResult XrDestroyActionSet();
+
+	XrResult XrCreateAction(
+		const XrActionCreateInfo* createInfo,
+		XrAction* action);
+	XrResult DestroyAction(FOXRVisionOSAction* Action);
+
+	XrResult Attach(FOXRVisionOSSession* Session);
+	bool IsAttached() const { return Session != nullptr; }
+
+	const TArray<TSharedPtr<FOXRVisionOSAction, ESPMode::ThreadSafe>>& GetActions() const { return Actions;  }
+	void ClearActive();
+	void SetActive(XrPath SubActionPath);
+	//uint32 GetPriority() const { return CreateInfo.priority; } //Note: action set priority currently unsupported
+	void SyncActions(TSharedPtr<FOXRVisionOSController, ESPMode::ThreadSafe>& Controllers, TSharedPtr <FOXRVisionOSTracker, ESPMode::ThreadSafe>& Tracker);
+
+	const FOXRVisionOSInstance* GetInstance() const { return Instance; }
+
+private:
+	bool bCreateFailed = false;
+	FOXRVisionOSInstance* Instance = nullptr;
+	FOXRVisionOSSession* Session = nullptr;
+	XrActionSetCreateInfo CreateInfo;
+	TArray<TSharedPtr<FOXRVisionOSAction, ESPMode::ThreadSafe>> Actions;
+	bool bAllActive = false;
+};
diff --git a/Engine/Platforms/VisionOS/Plugins/Runtime/OpenXRVisionOS/Source/OXRVisionOS/Private/OXRVisionOSAppDelegate.cpp b/Engine/Platforms/VisionOS/Plugins/Runtime/OpenXRVisionOS/Source/OXRVisionOS/Private/OXRVisionOSAppDelegate.cpp
new file mode 100644
--- /dev/null
+++ b/Engine/Platforms/VisionOS/Plugins/Runtime/OpenXRVisionOS/Source/OXRVisionOS/Private/OXRVisionOSAppDelegate.cpp
@@ -0,0 +1,19 @@
+// Copyright Epic Games, Inc. All Rights Reserved.
+
+#include "OXRVisionOSAppDelegate.h"
+#include "IOS/IOSAppDelegate.h"
+
+CP_OBJECT_cp_layer_renderer* OXRVisionOS::GetSwiftLayerRenderer()
+{
+	return [IOSAppDelegate GetDelegate].SwiftLayer;
+}
+
+int OXRVisionOS::GetSwiftNumViewports()
+{
+	return  [IOSAppDelegate GetDelegate].SwiftLayerViewports.count;
+}
+CGRect OXRVisionOS::GetSwiftViewportRect(int Index)
+{
+	return  [[[IOSAppDelegate GetDelegate].SwiftLayerViewports objectAtIndex:Index] CGRectValue];
+}
+
diff --git a/Engine/Platforms/VisionOS/Plugins/Runtime/OpenXRVisionOS/Source/OXRVisionOS/Private/OXRVisionOSAppDelegate.h b/Engine/Platforms/VisionOS/Plugins/Runtime/OpenXRVisionOS/Source/OXRVisionOS/Private/OXRVisionOSAppDelegate.h
new file mode 100644
--- /dev/null
+++ b/Engine/Platforms/VisionOS/Plugins/Runtime/OpenXRVisionOS/Source/OXRVisionOS/Private/OXRVisionOSAppDelegate.h
@@ -0,0 +1,14 @@
+// Copyright Epic Games, Inc. All Rights Reserved.
+
+#pragma once
+
+#import <CompositorServices/CompositorServices.h>
+
+namespace OXRVisionOS
+{
+	CP_OBJECT_cp_layer_renderer* GetSwiftLayerRenderer();
+
+	int GetSwiftNumViewports();
+
+	CGRect GetSwiftViewportRect(int Index);
+}
diff --git a/Engine/Platforms/VisionOS/Plugins/Runtime/OpenXRVisionOS/Source/OXRVisionOS/Private/OXRVisionOSController.cpp b/Engine/Platforms/VisionOS/Plugins/Runtime/OpenXRVisionOS/Source/OXRVisionOS/Private/OXRVisionOSController.cpp
new file mode 100644
--- /dev/null
+++ b/Engine/Platforms/VisionOS/Plugins/Runtime/OpenXRVisionOS/Source/OXRVisionOS/Private/OXRVisionOSController.cpp
@@ -0,0 +1,236 @@
+// Copyright Epic Games, Inc. All Rights Reserved.
+
+#include "OXRVisionOSController.h"
+#include "OXRVisionOSSession.h"
+#include "OXRVisionOSInstance.h"
+#include "OXRVisionOSTracker.h"
+#include "OXRVisionOSPlatformUtils.h"
+#include "Math/UnrealMathUtility.h"
+//#include "HAL/PlatformMemory.h"
+#include "Runtime/Launch/Resources/Version.h"
+
+namespace OXRVisionOS
+{
+	const TCHAR* OXRVisionOSControllerButtonEnumNames[] =
+	{
+		TEXT("PushL"),
+
+		TEXT("PushR"),
+
+		TEXT("GripL"),
+		TEXT("AimL"),
+		TEXT("GripR"),
+		TEXT("AimR"),
+		TEXT("HMDPose"),
+
+//		TEXT("GazePose"),
+
+		TEXT("NullInput"),
+	};
+	static_assert((int32)(EOXRVisionOSControllerButton::NullInput) == UE_ARRAY_COUNT(OXRVisionOSControllerButtonEnumNames), "OXRVisionOSControllerButtonEnumNames does not match EOXRVisionOSControllerButton!");
+
+	const TCHAR* OXRVisionOSControllerButtonToTCHAR(EOXRVisionOSControllerButton Button)
+	{
+		check((int32)Button >= 0);
+		check((int32)Button <= (int32)EOXRVisionOSControllerButton::NullInput);
+		return OXRVisionOSControllerButtonEnumNames[(int32)Button];
+	}
+}
+
+//=============================================================================
+bool VrControllerAPI::Init()
+{
+	return true;
+}
+
+//=============================================================================
+void VrControllerAPI::Term()
+{
+}
+
+//=============================================================================
+int32 VrControllerAPI::ControllerOpen(EControllerHand ControllerHand)
+{
+	UE_LOG(LogOXRVisionOS, Warning, TEXT("VrControllerAPI::ControllerOpen() : %d"), static_cast<int32>(ControllerHand));
+
+	return static_cast<int32>(ControllerHand);
+}
+
+//=============================================================================
+void VrControllerAPI::ControllerClose(int32 Handle)
+{
+}
+
+//=============================================================================
+bool VrControllerAPI::ControllerIsReady(int32 Handle)
+{
+	return true;
+}
+
+
+
+/** FOXRVisionOSController implementation */
+bool FOXRVisionOSController::Create(TSharedPtr<FOXRVisionOSController, ESPMode::ThreadSafe>& OutController, FOXRVisionOSSession* InSession)
+{
+	OutController = MakeShared<FOXRVisionOSController, ESPMode::ThreadSafe>(InSession);
+	if (OutController->bCreateFailed)
+	{
+		OutController = nullptr;
+		return false;
+	}
+
+	return true;
+}
+
+//=============================================================================
+FOXRVisionOSController::FOXRVisionOSController(FOXRVisionOSSession* InSession)
+{
+	Session = InSession;
+
+	FMemory::Memzero(&UserState, sizeof(UserState));
+
+	VrControllerAPI::Init();
+}
+
+//=============================================================================
+FOXRVisionOSController::~FOXRVisionOSController()
+{
+	VrControllerAPI::Term();
+}
+
+void FOXRVisionOSController::OnBeginSession(FOXRVisionOSTracker* Tracker)
+{
+	for (int i = 0; i < 2; ++i)
+	{
+		FControllerState& State = UserState.ControllerStates[i];
+
+		const EControllerHand ControllerHand = ControllerIndexToHand(i);
+		State.Handle = VrControllerAPI::ControllerOpen(ControllerHand);
+		State.Hand = ControllerHand;
+
+		//TODO register the device
+		//Tracker->RegisterDevice(TrackerDeviceType, State.Handle);
+	}
+}
+
+void FOXRVisionOSController::OnEndSession(FOXRVisionOSTracker* Tracker)
+{
+	for (int i = 0; i < 2; ++i)
+	{
+		const FControllerState& State = UserState.ControllerStates[i];
+		Tracker->UnregisterDevice(State.Handle);
+		VrControllerAPI::ControllerClose(State.Handle);
+	}
+}
+
+int32 FOXRVisionOSController::GetLeftControllerHandle() const
+{
+	const FControllerState& State = UserState.ControllerStates[0];
+	return State.Handle;
+}
+
+int32 FOXRVisionOSController::GetRightControllerHandle() const
+{
+	const FControllerState& State = UserState.ControllerStates[1];
+	return State.Handle;
+}
+
+void FOXRVisionOSController::SyncActions()
+{
+	for (FControllerState& ControllerState : UserState.ControllerStates)
+	{
+		//int32_t Result;
+
+		//Result = ReadControllerState(ControllerState.Handle, &ControllerState.);
+		//checkf(Result == RESULT_OK, TEXT("ControllerReadState failed Result=%d"), Result);
+	}
+}
+
+bool FOXRVisionOSController::GetActionStateBoolean(EOXRVisionOSControllerButton Button, bool& OutValue, bool &OutActive, XrTime& OutTimeStamp)
+{
+	FControllerState* ControllerState = nullptr;
+	bool bIsTouch = false;
+	switch (Button)
+	{
+	case EOXRVisionOSControllerButton::PushL:
+		ControllerState = &UserState.ControllerStates[0];
+		break;
+	case EOXRVisionOSControllerButton::PushR:
+		ControllerState = &UserState.ControllerStates[1];
+		break;
+	default:
+		UE_LOG(LogOXRVisionOS, Error, TEXT("GetActionStateBoolean unknown button %s"), OXRVisionOS::OXRVisionOSControllerButtonToTCHAR(Button));
+		OutActive = false;
+		check(false); // We should not fail to find a button here.
+		return false;
+	}
+
+	OutActive = false; //TODO read the controller state ControllerState->connected;
+	if (!OutActive)
+	{
+		return true;
+	}
+
+	//TODO: read the input state
+	OutValue = false;
+	return true;
+}
+
+// We are supporting treating boolean inputs as floats.
+bool FOXRVisionOSController::GetActionStateBooleanAsFloat(EOXRVisionOSControllerButton Button, float& OutValue, bool& OutActive, XrTime& OutTimeStamp)
+{
+	bool bOutValue = false;
+	if (GetActionStateBoolean(Button, bOutValue, OutActive, OutTimeStamp))
+	{
+		OutValue = bOutValue ? 1.0f : 0.0f;
+		return true;
+	}
+	OutActive = false;
+	return false;
+}
+
+#define TRIGGER_TO_FLOAT(x) ((float)(x) / 255.f)
+#define ANALOG_TO_FLOAT(x) (((float)(x) - 128.0f) / 128.f)
+
+bool FOXRVisionOSController::GetActionStateFloat(EOXRVisionOSControllerButton Button, float& OutValue, bool& OutActive, XrTime& OutTimeStamp)
+{
+	return GetActionStateBooleanAsFloat(Button, OutValue, OutActive, OutTimeStamp);
+}
+
+bool FOXRVisionOSController::GetActionStateVector2f(EOXRVisionOSControllerButton Button, XrVector2f& OutValue, bool& OutActive, XrTime& OutTimeStamp)
+{
+	OutActive = false;
+	check(false);
+	return false;
+}
+
+int32 FOXRVisionOSController::ControllerHandToIndex(EControllerHand ControllerHand)
+{
+	switch (ControllerHand)
+	{
+	default:
+		checkNoEntry();
+	case EControllerHand::Left:
+		return LeftHandIndex;
+	case EControllerHand::Right:
+		return RightHandIndex;
+	}
+}
+
+EControllerHand FOXRVisionOSController::ControllerIndexToHand(int32 ControllerIndex)
+{
+	switch (ControllerIndex)
+	{
+	default:
+		checkNoEntry();
+	case LeftHandIndex:
+		return EControllerHand::Left;
+	case RightHandIndex:
+		return EControllerHand::Right;
+	}
+}
+
+bool FOXRVisionOSController::IsFloatButton(EOXRVisionOSControllerButton Button)
+{
+	return false;
+}
diff --git a/Engine/Platforms/VisionOS/Plugins/Runtime/OpenXRVisionOS/Source/OXRVisionOS/Private/OXRVisionOSController.h b/Engine/Platforms/VisionOS/Plugins/Runtime/OpenXRVisionOS/Source/OXRVisionOS/Private/OXRVisionOSController.h
new file mode 100644
--- /dev/null
+++ b/Engine/Platforms/VisionOS/Plugins/Runtime/OpenXRVisionOS/Source/OXRVisionOS/Private/OXRVisionOSController.h
@@ -0,0 +1,119 @@
+// Copyright Epic Games, Inc. All Rights Reserved.
+
+#pragma once
+
+#include "CoreMinimal.h"
+#include "OpenXRCore.h"
+
+class FOXRVisionOSSession;
+class FOXRVisionOSTracker;
+struct FHapticFeedbackBuffer;
+struct FInputDeviceProperty;
+enum class EControllerHand : uint8;
+
+enum class EOXRVisionOSControllerButton : int32
+{
+	LeftButtonStart = 0,
+	PushL = LeftButtonStart,
+	LeftButtonCount,
+
+	RightButtonStart = LeftButtonCount,
+	PushR = RightButtonStart,
+	RightButtonCount,
+
+	PhysicalButtonCount = RightButtonCount,
+
+	/** Max number of controller buttons */
+	TotalButtonCount,
+
+	/** Pose inputs, handled by OXRVisionOSTracker **/
+	GripL = TotalButtonCount,
+	AimL,
+	GripR,
+	AimR,
+	HMDPose,
+
+	/** The eye tracking input **/
+//	GazePose,
+
+	NullInput, // This input is never mapped to hardware.
+};
+
+namespace OXRVisionOS
+{
+	const TCHAR* OXRVisionOSControllerButtonToTCHAR(EOXRVisionOSControllerButton Button);
+}
+
+class VrControllerAPI
+{
+public:
+
+	struct EVrControllerButton
+	{
+		static const int32 Mask[(int32)EOXRVisionOSControllerButton::PhysicalButtonCount];
+	};
+
+	/** VrControllerAPI Initialization and termination */
+	static bool Init();
+	static void Term();
+
+	/** Controller methods */
+	static int32 ControllerOpen(EControllerHand ControllerHand);
+	static void ControllerClose(int32 Handle);
+	static bool ControllerIsReady(int32 Handle);
+};
+
+/**
+ * Implementation of the OXRVisionOS motion controllers.
+ * One instance of this handles both controllers.
+ */
+class FOXRVisionOSController
+{
+public:
+	static bool Create(TSharedPtr<FOXRVisionOSController, ESPMode::ThreadSafe>& OutController, FOXRVisionOSSession* InSession);
+	FOXRVisionOSController(FOXRVisionOSSession* InSession);
+	~FOXRVisionOSController();
+
+	void OnBeginSession(FOXRVisionOSTracker* Tracker);
+	void OnEndSession(FOXRVisionOSTracker* Tracker);
+
+	void SyncActions();
+	bool GetActionStateBoolean(EOXRVisionOSControllerButton Button, bool& OutValue, bool& OutActive, XrTime& OutTimeStamp);
+	bool GetActionStateFloat(EOXRVisionOSControllerButton Button, float& OutValue, bool& OutActive, XrTime& OutTimeStamp);
+	bool GetActionStateVector2f(EOXRVisionOSControllerButton Button, XrVector2f& OutValue, bool& OutActive, XrTime& OutTimeStamp);
+
+	int32 GetLeftControllerHandle() const;
+	int32 GetRightControllerHandle() const;
+	static bool IsFloatButton(EOXRVisionOSControllerButton Button);
+	
+private:
+	static int32 ControllerHandToIndex(EControllerHand ControllerHand);
+	static EControllerHand ControllerIndexToHand(int32 ControllerIndex);
+
+	bool bCreateFailed = false;
+	FOXRVisionOSSession* Session = nullptr;
+
+	/** State for a controller */
+	struct FControllerState
+	{
+		//TODO
+		int32 Handle;
+
+		/** Left or right hand which the controller represents */
+		EControllerHand Hand;
+	};
+
+	static constexpr int32 LeftHandIndex = 0;
+	static constexpr int32 RightHandIndex = 1;
+
+	/** State for user connections */
+	struct FUserState
+	{
+		/** Per controller state, 2 per user for left & right hands */
+		FControllerState	ControllerStates[2];
+	} UserState;
+
+	/** Helper for using boolean inputs as float inputs*/
+	bool GetActionStateBooleanAsFloat(EOXRVisionOSControllerButton Button, float& OutValue, bool& OutActive, XrTime& OutTimeStamp);
+};
+
diff --git a/Engine/Platforms/VisionOS/Plugins/Runtime/OpenXRVisionOS/Source/OXRVisionOS/Private/OXRVisionOSInstance.cpp b/Engine/Platforms/VisionOS/Plugins/Runtime/OpenXRVisionOS/Source/OXRVisionOS/Private/OXRVisionOSInstance.cpp
new file mode 100644
--- /dev/null
+++ b/Engine/Platforms/VisionOS/Plugins/Runtime/OpenXRVisionOS/Source/OXRVisionOS/Private/OXRVisionOSInstance.cpp
@@ -0,0 +1,753 @@
+// Copyright Epic Games, Inc. All Rights Reserved.
+
+#include "OXRVisionOSInstance.h"
+#include "OXRVisionOSSession.h"
+#include "OXRVisionOSActionSet.h"
+#include "OXRVisionOSAction.h"
+#include "OXRVisionOSController.h"
+#include "Engine/Engine.h"
+#include "OXRVisionOS_openxr_platform.h"
+#include "OXRVisionOSAppDelegate.h"
+#include "Misc/ConfigCacheIni.h"
+#include "MetalRHIVisionOSBridge.h"
+
+#define OXRVISIONOS_INSTANCE_NAME "OXRVisionOS"
+
+// OXRVisionOS only supports one system.
+#define OXRVISIONOS_XR_SYSTEM_ID 1
+#define OXRVISIONOS_XR_SYSTEM_NAME "OXRVisionOSSystem"
+
+#define OXRVISIONOS_RUNTIME_VERSION XR_MAKE_VERSION(0, 1, 0)
+
+// Note: This is not a real registered vendorid, it is the khronos id + 1.
+#define OXRVISIONOS_VENDOR_ID 0x10001
+
+TSharedPtr<FOXRVisionOSSession, ESPMode::ThreadSafe> FOXRVisionOSInstance::Session;
+
+namespace OXRVisionOSInstanceHelpers
+{
+	constexpr uint32 RecommendedPerEyeSwapchainWidth = 512;
+	constexpr uint32 RecommendedPerEyeSwapchainHeight = 512;
+
+	constexpr uint32 DefaultMaxPerEyeSwapchainWidth = 512;
+	constexpr uint32 DefaultMaxPerEyeSwapchainHeight = 512;
+}
+
+XrResult FOXRVisionOSInstance::Create(TSharedPtr<FOXRVisionOSInstance, ESPMode::ThreadSafe>& OutInstance, const XrInstanceCreateInfo* CreateInfo, FOXRVisionOS* InModule)
+{
+	if (CreateInfo == nullptr || CreateInfo->type != XR_TYPE_INSTANCE_CREATE_INFO)
+	{
+		return XrResult::XR_ERROR_VALIDATION_FAILURE;
+	}
+
+	const XrApplicationInfo& AppInfo = CreateInfo->applicationInfo;
+
+	if (AppInfo.applicationName[0] == 0)
+	{
+		return XrResult::XR_ERROR_NAME_INVALID;
+	}
+
+	//Could check api version here
+	//createInfo->
+	//XR_ERROR_API_VERSION_UNSUPPORTED
+
+	//Could check/handle API layers
+	//Note API layers generally do not require special handling.
+	//uint32_t					createInfo.enabledApiLayerCount;
+	//const char* const*		createInfo.enabledApiLayerNames;
+	//XR_ERROR_API_LAYER_NOT_PRESENT
+
+	const FOXRVisionOS* OXRVisionOS = FOXRVisionOS::Get();
+	check(OXRVisionOS);
+	for (int i = 0; i < CreateInfo->enabledExtensionCount; i++)
+	{
+		const char* ExtensionName = CreateInfo->enabledExtensionNames[i];
+		if (!OXRVisionOS->IsExtensionSupported(ExtensionName))
+		{
+			UE_LOG(LogOXRVisionOS, Error, TEXT("FOXRVisionOSInstance::Create required extension: %hs but it is not supported.  Instance creation failed."), ExtensionName);
+			return XrResult::XR_ERROR_EXTENSION_NOT_PRESENT;
+		}
+	}
+
+	OutInstance = MakeShared<FOXRVisionOSInstance, ESPMode::ThreadSafe>(CreateInfo, InModule);
+	if (OutInstance->bCreateFailed)
+	{
+		OutInstance = nullptr;
+		return XrResult::XR_ERROR_RUNTIME_FAILURE;
+	}
+	return XrResult::XR_SUCCESS;
+}
+
+FOXRVisionOSInstance::FOXRVisionOSInstance(const XrInstanceCreateInfo* CreateInfo, FOXRVisionOS* InModule)
+{
+	Module = InModule;
+
+	// Enable all requested extensions.  We already failed if any were unsupported.
+	for (int i = 0; i < CreateInfo->enabledExtensionCount; i++)
+	{
+		const char* ExtensionName = CreateInfo->enabledExtensionNames[i];
+		EnabledExtensions.Add(ExtensionName);
+	}
+
+	// Cache instance properties.
+	InstanceProperties.type = XR_TYPE_INSTANCE_PROPERTIES;
+	InstanceProperties.next = nullptr;
+	InstanceProperties.runtimeVersion = OXRVISIONOS_RUNTIME_VERSION;
+	FCStringAnsi::Strncpy(InstanceProperties.runtimeName, OXRVISIONOS_INSTANCE_NAME, XR_MAX_RUNTIME_NAME_SIZE);
+
+	// Cache system properties.
+	SystemProperties.type = XR_TYPE_SYSTEM_PROPERTIES;
+	SystemProperties.next = nullptr;
+	SystemProperties.systemId = OXRVISIONOS_XR_SYSTEM_ID;
+	SystemProperties.vendorId = OXRVISIONOS_VENDOR_ID;
+	FCStringAnsi::Strncpy(SystemProperties.systemName, OXRVISIONOS_XR_SYSTEM_NAME, XR_MAX_SYSTEM_NAME_SIZE);
+
+	//TODO get graphics properties here
+    const int NumViewports = OXRVisionOS::GetSwiftNumViewports();
+    check(NumViewports > 0);
+    check(NumViewports < 3); // unexpected, and we aren't handling it.
+    CGRect LeftRect = OXRVisionOS::GetSwiftViewportRect(0);
+    CGRect RightRect = NumViewports > 1 ? OXRVisionOS::GetSwiftViewportRect(1) : CGRectMake(0,0,LeftRect.size.width/4, LeftRect.size.height/4); //TODO CVars around this perhaps???
+    //LeftRect.size.width, LeftRect.size.height
+
+	int32 MaxWidth = LeftRect.size.width + (NumViewports > 1 ? RightRect.size.width : 0);
+	int32 MaxHeight = LeftRect.size.height;
+	//if (GConfig->GetInt(TEXT("/Script/OXRVisionOSSettings.OXRVisionOSRuntimeSettings"), TEXT("OXRVisionOSMaxPerEyeWidth"), MaxWidth, GEngineIni))
+	//{
+	//	UE_LOG(LogOXRVisionOS, Log, TEXT("Max texture width read from settings file: %d"), MaxWidth);
+	//}
+	//if (GConfig->GetInt(TEXT("/Script/OXRVisionOSSettings.OXRVisionOSRuntimeSettings"), TEXT("OXRVisionOSMaxPerEyeHeight"), MaxHeight, GEngineIni))
+	//{
+	//	UE_LOG(LogOXRVisionOS, Log, TEXT("Max texture height read from settings file: %d"), MaxHeight);
+	//}
+
+	SystemProperties.graphicsProperties.maxSwapchainImageWidth = MaxWidth;
+    SystemProperties.graphicsProperties.maxSwapchainImageHeight = MaxHeight;
+	//// OXRVisionOS version of OpenXR is not-compliant because we support less than XR_MIN_COMPOSITION_LAYERS_SUPPORTED
+	//TODO temporarily just made up this value, but it may be correct.
+	SystemProperties.graphicsProperties.maxLayerCount = 1;
+
+	SystemProperties.trackingProperties.orientationTracking = XR_TRUE;
+	SystemProperties.trackingProperties.positionTracking = XR_TRUE;
+
+	// Create the null path
+	{
+		const char PathString[] = "XR_NULL_PATH";
+		const SIZE_T PathSize = static_cast<SIZE_T>(sizeof(PathString));
+		TUniquePtr<ANSICHAR[]> PathArrayElementPtr = MakeUnique<ANSICHAR[]>(PathSize);
+		ANSICHAR* const CharPtr = PathArrayElementPtr.Get();
+		TCString<ANSICHAR>::Strncpy(CharPtr, PathString, PathSize);
+		uint32 NewIndex = XrPaths.Add(MoveTemp(PathArrayElementPtr));
+		const XrPath NewPath = NewIndex;
+		StringToXrPathMap.Add(CharPtr, NewPath);
+	}
+
+	//// Add additional system properties for extensions insert each one after the head struct.
+	//if (IsExtensionEnabled(XR_EXT_EYE_GAZE_INTERACTION_EXTENSION_NAME))
+	//{
+	//	// Eye Gaze interaction
+	//	EyeGazeInteractionSystemProperties.type = XR_TYPE_SYSTEM_EYE_GAZE_INTERACTION_PROPERTIES_EXT;
+	//	EyeGazeInteractionSystemProperties.supportsEyeGazeInteraction = true;
+	//	EyeGazeInteractionSystemProperties.next = SystemProperties.next;
+	//	SystemProperties.next = &EyeGazeInteractionSystemProperties;
+	//}
+
+	OnWorldTickStartDelegateHandle = FWorldDelegates::OnWorldTickStart.AddRaw(this, &FOXRVisionOSInstance::OnWorldTickStart);
+}
+
+FOXRVisionOSInstance::~FOXRVisionOSInstance()
+{
+	FWorldDelegates::OnWorldTickStart.Remove(OnWorldTickStartDelegateHandle);
+
+	if (bCreateFailed)
+	{
+		UE_LOG(LogOXRVisionOS, Warning, TEXT("Destructing FOXRVisionOSInstance because instance create failed."));
+	}
+	else
+	{
+		//TODO instance loss logic?
+	}
+}
+
+void FOXRVisionOSInstance::OnWorldTickStart(UWorld* World, ELevelTick TickType, float DeltaTime)
+{
+	//TODO
+	// Poll device state to find out if the HMD is on the user's head
+	// and send delegates about hmd worn status.  OpenXR hides this specific info from applications, the runtime is supposed to handle it.
+	// But here we are the runtime so we need to make it possible to deal with this.
+	// We may need to drive the session to idle when the hmd is off, and return it to ready when the hmd is on.   See FOXRVisionOSSession:OnWorldTickStart.
+	// 		const bool NewbHMDWorn = TODO_PLATFORM_API_HMD_WORN;
+	// 		if (NewbHMDWorn != bHMDWorn)
+	// 		{
+	// 			bHMDWorn = NewbHMDWorn;
+	// 			if (bHMDWorn)
+	// 			{
+	// 				FCoreDelegates::VRHeadsetPutOnHead.Broadcast();
+	// 			}
+	// 			else
+	// 			{
+	// 				FCoreDelegates::VRHeadsetRemovedFromHead.Broadcast();
+	// 			}
+	// 		}
+}
+
+XrResult FOXRVisionOSInstance::XrGetSystem(
+	const XrSystemGetInfo* getInfo,
+	XrSystemId* systemId)
+{
+	check(bCreateFailed == false);
+
+	if (getInfo == nullptr || getInfo->type != XR_TYPE_SYSTEM_GET_INFO)
+	{
+		return XrResult::XR_ERROR_VALIDATION_FAILURE;
+	}
+
+	if (getInfo->formFactor != XR_FORM_FACTOR_HEAD_MOUNTED_DISPLAY)
+	{
+		return XrResult::XR_ERROR_FORM_FACTOR_UNSUPPORTED;
+	}
+
+	if (!IsHMDAvailable())
+	{
+		return XR_ERROR_FORM_FACTOR_UNAVAILABLE;
+	}
+
+	*systemId = OXRVISIONOS_XR_SYSTEM_ID;
+	return XrResult::XR_SUCCESS;
+}
+
+bool FOXRVisionOSInstance::IsHMDAvailable()
+{
+	// We are running on the device which is the hmd, so it can't be unavailable.
+	return true;
+}
+
+XrResult FOXRVisionOSInstance::XrGetInstanceProperties(
+	XrInstanceProperties* OutInstanceProperties)
+{
+	if (OutInstanceProperties == nullptr)
+	{
+		return XrResult::XR_ERROR_VALIDATION_FAILURE;
+	}
+
+	*OutInstanceProperties = InstanceProperties;
+	return XrResult::XR_SUCCESS;
+}
+
+XrResult FOXRVisionOSInstance::XrGetSystemProperties(
+	XrSystemId                                  SystemId,
+	XrSystemProperties*							Properties)
+{
+	if (SystemId != OXRVISIONOS_XR_SYSTEM_ID)
+	{
+		return XrResult::XR_ERROR_SYSTEM_INVALID;
+	}
+	if (Properties == nullptr)
+	{
+		return XrResult::XR_ERROR_VALIDATION_FAILURE;
+	}
+
+	*Properties = SystemProperties;
+	return XrResult::XR_SUCCESS;
+}
+
+XrResult FOXRVisionOSInstance::XrEnumerateEnvironmentBlendModes(
+	XrSystemId                                  SystemId,
+	XrViewConfigurationType                     ViewConfigurationType,
+	uint32_t                                    EnvironmentBlendModeCapacityInput,
+	uint32_t*									EnvironmentBlendModeCountOutput,
+	XrEnvironmentBlendMode*						EnvironmentBlendModes)
+{
+	if (SystemId != OXRVISIONOS_XR_SYSTEM_ID)
+	{
+		return XrResult::XR_ERROR_SYSTEM_INVALID;
+	}
+	if (EnvironmentBlendModeCountOutput == nullptr)
+	{
+		return XrResult::XR_ERROR_VALIDATION_FAILURE;
+	}
+	if (EnvironmentBlendModeCapacityInput != 0 && EnvironmentBlendModes == nullptr)
+	{
+		return XrResult::XR_ERROR_VALIDATION_FAILURE;
+	}
+
+	switch (ViewConfigurationType)
+	{
+	//case XR_VIEW_CONFIGURATION_TYPE_PRIMARY_MONO: // TODO add or delete this.
+	case  XrViewConfigurationType::XR_VIEW_CONFIGURATION_TYPE_PRIMARY_STEREO:
+	{
+		*EnvironmentBlendModeCountOutput = 1;
+		if (EnvironmentBlendModeCapacityInput != 0)
+		{
+			check(EnvironmentBlendModeCapacityInput >= 1); // Because we only need 1 we can't ever need to return XR_ERROR_SIZE_INSUFFICIENT
+
+			EnvironmentBlendModes[0] = XR_ENVIRONMENT_BLEND_MODE_OPAQUE;
+		}
+		return XrResult::XR_SUCCESS;
+	}
+	default:
+	{
+		*EnvironmentBlendModeCountOutput = 0;
+		return XrResult::XR_ERROR_VIEW_CONFIGURATION_TYPE_UNSUPPORTED;
+	}
+	}
+}
+
+XrResult FOXRVisionOSInstance::XrEnumerateViewConfigurations(
+	XrSystemId                                  SystemId,
+	uint32_t                                    ViewConfigurationTypeCapacityInput,
+	uint32_t*									ViewConfigurationTypeCountOutput,
+	XrViewConfigurationType*					ViewConfigurationTypes)
+{
+	if (SystemId != OXRVISIONOS_XR_SYSTEM_ID)
+	{
+		return XrResult::XR_ERROR_SYSTEM_INVALID;
+	}
+	if (ViewConfigurationTypeCountOutput == nullptr)
+	{
+		return XrResult::XR_ERROR_VALIDATION_FAILURE;
+	}
+	if (ViewConfigurationTypeCapacityInput != 0 && ViewConfigurationTypes == nullptr)
+	{
+		return XrResult::XR_ERROR_VALIDATION_FAILURE;
+	}
+
+	*ViewConfigurationTypeCountOutput = 1;
+	if (ViewConfigurationTypeCapacityInput != 0)
+	{
+		check(ViewConfigurationTypeCapacityInput >= 1); // Because we only need 1 we can't ever need to return XR_ERROR_SIZE_INSUFFICIENT
+
+		ViewConfigurationTypes[0] = XrViewConfigurationType::XR_VIEW_CONFIGURATION_TYPE_PRIMARY_STEREO;
+	}
+	return XrResult::XR_SUCCESS;
+}
+
+XrResult FOXRVisionOSInstance::XrGetViewConfigurationProperties(
+	XrSystemId                                  SystemId,
+	XrViewConfigurationType                     ViewConfigurationType,
+	XrViewConfigurationProperties*				ConfigurationProperties)
+{
+	if (SystemId != OXRVISIONOS_XR_SYSTEM_ID)
+	{
+		return XrResult::XR_ERROR_SYSTEM_INVALID;
+	}
+	if (ConfigurationProperties == nullptr)
+	{
+		return XrResult::XR_ERROR_VALIDATION_FAILURE;
+	}
+
+	if (ViewConfigurationType == XrViewConfigurationType::XR_VIEW_CONFIGURATION_TYPE_PRIMARY_STEREO)
+	{
+		ConfigurationProperties->type = XR_TYPE_VIEW_CONFIGURATION_PROPERTIES;
+		ConfigurationProperties->next = nullptr;
+		ConfigurationProperties->viewConfigurationType = XrViewConfigurationType::XR_VIEW_CONFIGURATION_TYPE_PRIMARY_STEREO;
+		ConfigurationProperties->fovMutable = XR_FALSE;
+		return XrResult::XR_SUCCESS;
+	}
+	else
+	{
+		return XrResult::XR_ERROR_VIEW_CONFIGURATION_TYPE_UNSUPPORTED;
+	}
+}
+
+XrResult FOXRVisionOSInstance::XrEnumerateViewConfigurationViews(
+	XrSystemId                                  SystemId,
+	XrViewConfigurationType                     ViewConfigurationType,
+	uint32_t                                    ViewCapacityInput,
+	uint32_t*									ViewCountOutput,
+	XrViewConfigurationView*					Views)
+{
+	if (SystemId != OXRVISIONOS_XR_SYSTEM_ID)
+	{
+		return XrResult::XR_ERROR_SYSTEM_INVALID;
+	}
+	if (ViewCountOutput == nullptr)
+	{
+		return XrResult::XR_ERROR_VALIDATION_FAILURE;
+	}
+	if (ViewCapacityInput != 0 && Views == nullptr)
+	{
+		return XrResult::XR_ERROR_VALIDATION_FAILURE;
+	}
+
+	if (ViewConfigurationType == XR_VIEW_CONFIGURATION_TYPE_PRIMARY_STEREO)
+	{
+		*ViewCountOutput = 2;
+		if (ViewCapacityInput != 0)
+		{
+			if (ViewCapacityInput < *ViewCountOutput)
+			{
+				return XR_ERROR_SIZE_INSUFFICIENT;
+			}
+
+			XrViewConfigurationView View;
+			View.type = XR_TYPE_VIEW_CONFIGURATION_VIEW;
+			View.next = nullptr;
+			View.recommendedImageRectWidth			= SystemProperties.graphicsProperties.maxSwapchainImageWidth;
+			View.recommendedImageRectHeight			= SystemProperties.graphicsProperties.maxSwapchainImageHeight;
+			View.maxImageRectWidth					= SystemProperties.graphicsProperties.maxSwapchainImageWidth;
+			View.maxImageRectHeight					= SystemProperties.graphicsProperties.maxSwapchainImageHeight;
+			View.recommendedSwapchainSampleCount	= 1; //TODO: Is this right?
+			View.maxSwapchainSampleCount			= 1; //TODO: Is this right?
+
+			Views[0] = View;
+			Views[1] = View;
+		}
+		return XrResult::XR_SUCCESS;
+	}
+	else
+	{
+		return XrResult::XR_ERROR_VIEW_CONFIGURATION_TYPE_UNSUPPORTED;
+	}
+}
+
+XrResult FOXRVisionOSInstance::XrCreateSession(
+	const XrSessionCreateInfo*		CreateInfo,
+	XrSession*						OutSession)
+{
+	// Only one session supported.
+	if (Session.IsValid())
+	{
+		return XrResult::XR_ERROR_LIMIT_REACHED;
+	}
+
+	XrResult Ret = FOXRVisionOSSession::Create(Session, CreateInfo, this);
+	*OutSession = (XrSession)Session.Get();
+	return Ret;
+}
+
+XrResult FOXRVisionOSInstance::XrStringToPath(
+	const char*				PathString,
+	XrPath*					Path)
+{
+	if (PathString == nullptr)
+	{
+		return XrResult::XR_ERROR_VALIDATION_FAILURE;
+	}
+
+	// If this path exists return it.
+	XrPath* Found = StringToXrPathMap.Find(PathString);
+	if (Found != nullptr)
+	{
+		*Path = *Found;
+		return XrResult::XR_SUCCESS;
+	}
+
+	// Validate the pathString
+	if (PathString[0] != '/')
+	{
+		return XR_ERROR_PATH_FORMAT_INVALID;
+	}
+	uint32 PathSize = FCStringAnsi::Strlen(PathString) + 1;
+	if (PathSize > XR_MAX_PATH_LENGTH)
+	{
+		return XR_ERROR_PATH_FORMAT_INVALID;
+	}
+
+
+	// Create a new path and return it..
+	*Path = StringToPath(PathString, PathSize);
+	return XrResult::XR_SUCCESS;
+}
+
+XrPath FOXRVisionOSInstance::StringToPath(const char* PathString)
+{
+	uint32 PathSize = FCStringAnsi::Strlen(PathString) + 1;
+	return StringToPath(PathString, PathSize);
+}
+XrPath FOXRVisionOSInstance::StringToPath(const char* PathString, uint32 PathSize)
+{
+	check(PathSize <= XR_MAX_PATH_LENGTH);
+
+	XrPath* FoundPath = StringToXrPathMap.Find(PathString);
+	if (FoundPath)
+	{
+		return *FoundPath;
+	}
+
+	// Create a new path and return it..
+	TUniquePtr<ANSICHAR[]> PathArrayElementPtr = MakeUnique<ANSICHAR[]>(PathSize);
+	ANSICHAR* const CharPtr = PathArrayElementPtr.Get();
+	TCString<ANSICHAR>::Strncpy(CharPtr, PathString, PathSize);
+	uint32 NewIndex = XrPaths.Add(MoveTemp(PathArrayElementPtr));
+	const XrPath NewPath = NewIndex;
+	StringToXrPathMap.Add(CharPtr, NewPath);
+	return NewPath;
+}
+
+XrResult FOXRVisionOSInstance::XrPathToString(
+	XrPath                                      Path,
+	uint32_t									BufferCapacityInput,
+	uint32_t*									BufferCountOutput,
+	char*										Buffer)
+{
+	if (Path == XR_NULL_PATH)
+	{
+		return XrResult::XR_ERROR_HANDLE_INVALID;
+	}
+
+	const uint32_t Index = Path;
+
+	if (Index > XrPaths.Num())
+	{
+		return XrResult::XR_ERROR_HANDLE_INVALID;
+	}
+
+	const uint32_t PathSize = FCStringAnsi::Strlen(XrPaths[Index].Get()) + 1;
+	*BufferCountOutput = PathSize;
+	if (BufferCapacityInput == 0)
+	{
+		return XrResult::XR_SUCCESS;
+	}
+
+	if (BufferCapacityInput < PathSize)
+	{
+		return XrResult::XR_ERROR_SIZE_INSUFFICIENT;
+	}
+
+	if (Buffer == nullptr)
+	{
+		return XrResult::XR_ERROR_VALIDATION_FAILURE;
+	}
+
+	ANSICHAR* PathChars = XrPaths[Index].Get();
+	FCStringAnsi::Strncpy(Buffer, PathChars, BufferCapacityInput);
+
+	return XrResult::XR_SUCCESS;
+}
+
+XrResult FOXRVisionOSInstance::XrCreateActionSet(
+	const XrActionSetCreateInfo* CreateInfo,
+	XrActionSet* ActionSet)
+{
+	TSharedPtr<FOXRVisionOSActionSet, ESPMode::ThreadSafe> NewActionSet;
+	XrResult Ret = FOXRVisionOSActionSet::Create(NewActionSet, CreateInfo, this);
+	if (Ret == XrResult::XR_SUCCESS)
+	{
+		ActionSets.Add(NewActionSet);
+		*ActionSet = (XrActionSet)NewActionSet.Get();
+	}
+	return Ret;
+}
+
+XrResult FOXRVisionOSInstance::DestroyActionSet(class FOXRVisionOSActionSet* ActionSet)
+{
+	uint32 ArrayIndex = ActionSets.IndexOfByPredicate([ActionSet](const TSharedPtr<FOXRVisionOSActionSet, ESPMode::ThreadSafe>& Data) { return (Data.Get() == ActionSet); });
+
+	if (ArrayIndex == INDEX_NONE)
+	{
+		return  XrResult::XR_ERROR_HANDLE_INVALID;
+	}
+
+	ActionSets.RemoveAtSwap(ArrayIndex);
+
+	return XrResult::XR_SUCCESS;
+}
+
+TSharedPtr<FOXRVisionOSActionSet, ESPMode::ThreadSafe> FOXRVisionOSInstance::GetActionSetSharedPtr(XrActionSet ActionSet)
+{
+	for (TSharedPtr<FOXRVisionOSActionSet, ESPMode::ThreadSafe>& ActionSetPtr : ActionSets)
+	{
+		if (reinterpret_cast<XrActionSet>(ActionSetPtr.Get()) == ActionSet)
+		{
+			return ActionSetPtr;
+		}
+	}
+	return nullptr;
+}
+
+void FOXRVisionOSInstance::CreateInteractionProfiles()
+{
+	if (InteractionProfiles.Num() != 0)
+	{
+		return;
+	}
+
+	// Note: Input profiles are defined in order of decreasing preference.  when bindings for multiple profiles are provided each top level path will use the highest preference input profile.
+	// Thus visionos_controller goes first, any supported emulations should be in the middle with better emulations earlier, simple_controller and extensions go last.
+
+	if (IsExtensionEnabled(XR_EPIC_OXRVISIONOS_CONTROLLER_NAME))
+	{
+		const XrPath InteractionProfilePath = StringToPath("/interaction_profiles/epic/visionos_controller");
+		PreferredInteractionProfileBindingList.Add(InteractionProfilePath);
+		TMap<XrPath, EOXRVisionOSControllerButton>& Pairs = InteractionProfiles.Add(InteractionProfilePath);
+
+		// The native profile just maps the paths to their native source button.
+		Pairs.Add(StringToPath("/user/hand/left/input/select/click"), EOXRVisionOSControllerButton::PushL);  // stick click
+		Pairs.Add(StringToPath("/user/hand/right/input/select/click"), EOXRVisionOSControllerButton::PushR);
+		Pairs.Add(StringToPath("/user/hand/left/input/grip/pose"), EOXRVisionOSControllerButton::GripL);
+		Pairs.Add(StringToPath("/user/hand/right/input/grip/pose"), EOXRVisionOSControllerButton::GripR);
+		Pairs.Add(StringToPath("/user/hand/left/input/aim/pose"), EOXRVisionOSControllerButton::AimL);
+		Pairs.Add(StringToPath("/user/hand/right/input/aim/pose"), EOXRVisionOSControllerButton::AimR);
+
+		// Parent Paths  One can use the 'parent' path of a value or click and it should behave like click if there is one and value if there is not.  OpenXR spec 11.4.
+		Pairs.Add(StringToPath("/user/hand/left/input/select"), EOXRVisionOSControllerButton::PushL);			// stick click
+		Pairs.Add(StringToPath("/user/hand/right/input/select"), EOXRVisionOSControllerButton::PushR);
+	}
+
+	// Right now I'll assume we are too simple for simple controller.
+	//{
+	//	const XrPath InteractionProfilePath = StringToPath("/interaction_profiles/khr/simple_controller");
+	//	PreferredInteractionProfileBindingList.Add(InteractionProfilePath);
+	//	TMap<XrPath, EOXRVisionOSControllerButton>& Pairs = InteractionProfiles.Add(InteractionProfilePath);
+
+	//	Pairs.Add(StringToPath("/user/hand/left/input/select/click"), EOXRVisionOSControllerButton::PushL);
+	//	Pairs.Add(StringToPath("/user/hand/left/input/menu/click"), EOXRVisionOSControllerButton::NullInput);			// This breaks the simplecontroller input pretty badly have to figure something out for this to use this config
+	//	Pairs.Add(StringToPath("/user/hand/left/input/grip/pose"), EOXRVisionOSControllerButton::GripL);
+	//	Pairs.Add(StringToPath("/user/hand/left/input/aim/pose"), EOXRVisionOSControllerButton::AimL);
+	//	Pairs.Add(StringToPath("/user/hand/left/output/haptic"), EOXRVisionOSControllerButton::NullInput);				// haptics - we do not access haptics through these, so NullInput
+
+	//	Pairs.Add(StringToPath("/user/hand/right/input/select/click"), EOXRVisionOSControllerButton::PushR);
+	//	Pairs.Add(StringToPath("/user/hand/right/input/menu/click"), EOXRVisionOSControllerButton::NullInput);			// Intentionally not mapping a right button for this.
+	//	Pairs.Add(StringToPath("/user/hand/right/input/grip/pose"), EOXRVisionOSControllerButton::GripR);
+	//	Pairs.Add(StringToPath("/user/hand/right/input/aim/pose"), EOXRVisionOSControllerButton::AimR);
+	//	Pairs.Add(StringToPath("/user/hand/right/output/haptic"), EOXRVisionOSControllerButton::NullInput);				// haptics - we do not access haptics through these, so NullInput
+
+
+	//	// Parent Paths  One can use the 'parent' path of a value or click and it should behave like click if there is one and value if there is not.  OpenXR spec 11.4.
+	//	Pairs.Add(StringToPath("/user/hand/left/input/select"), EOXRVisionOSControllerButton::Square);
+	//	Pairs.Add(StringToPath("/user/hand/left/input/menu"), EOXRVisionOSControllerButton::NullInput);
+
+	//	Pairs.Add(StringToPath("/user/hand/right/input/select"), EOXRVisionOSControllerButton::Cross);
+	//	Pairs.Add(StringToPath("/user/hand/right/input/menu"), EOXRVisionOSControllerButton::NullInput);				// Intentionally not mapping a right button for this.
+	//}
+
+	// if (IsExtensionEnabled(XR_EXT_EYE_GAZE_INTERACTION_EXTENSION_NAME))
+	// {
+	// 	const XrPath InteractionProfilePath = StringToPath("/interaction_profiles/ext/eye_gaze_interaction");
+	// 	PreferredInteractionProfileBindingList.Add(InteractionProfilePath);
+	// 	TMap<XrPath, EOXRVisionOSControllerButton>& Pairs = InteractionProfiles.Add(InteractionProfilePath);
+
+	// 	Pairs.Add(StringToPath("/user/eyes_ext/input/gaze_ext/pose"), EOXRVisionOSControllerButton::GazePose);
+	// }
+
+	// Define all supported top level paths
+	TopLevelPathInteractionProfileMap.Add(StringToPath("/user/hand/left"), XR_NULL_PATH);
+	TopLevelPathInteractionProfileMap.Add(StringToPath("/user/hand/right"), XR_NULL_PATH);
+	TopLevelPathInteractionProfileMap.Add(StringToPath("/user/head"), XR_NULL_PATH);
+	TopLevelPathInteractionProfileMap.Add(StringToPath("/user/eyes_ext"), XR_NULL_PATH);
+
+	// Build map of all input paths to top level paths
+	for (const auto& Pair : TopLevelPathInteractionProfileMap)
+	{
+		const XrPath TopLevelPath = Pair.Key;
+		const ANSICHAR* TopLevelPathChars = XrPaths[TopLevelPath].Get();
+		const int32 TopLevelPathCharsLen = FCStringAnsi::Strlen(TopLevelPathChars);
+
+		for (const auto& Profile : InteractionProfiles)
+		{
+			for (const auto& Pair2 : Profile.Value)
+			{
+				XrPath InputPath = Pair2.Key;
+				if (!InputPathToTopLevelPathMap.Contains(InputPath))
+				{
+					ANSICHAR* InputPathChars = XrPaths[InputPath].Get();
+					if (FCStringAnsi::Strncmp(TopLevelPathChars, InputPathChars, TopLevelPathCharsLen) == 0)
+					{
+						InputPathToTopLevelPathMap.Add(InputPath, TopLevelPath);
+					}
+				}
+			}
+		}
+	}
+}
+
+bool FOXRVisionOSInstance::IsFirstInputProfilePreferred(XrPath First, XrPath Second) const
+{
+	if (Second == XR_NULL_PATH)
+	{
+		return true;
+	}
+	int32 FirstIndex = PreferredInteractionProfileBindingList.Find(First);
+	int32 SecondIndex = PreferredInteractionProfileBindingList.Find(Second);
+	check(FirstIndex != INDEX_NONE);
+	check(SecondIndex != INDEX_NONE);
+	return FirstIndex <= SecondIndex;
+}
+
+XrResult FOXRVisionOSInstance::XrSuggestInteractionProfileBindings(
+	const XrInteractionProfileSuggestedBinding* SuggestedBindings)
+{
+	if (SuggestedBindings == nullptr 
+		|| SuggestedBindings->type != XR_TYPE_INTERACTION_PROFILE_SUGGESTED_BINDING
+		|| SuggestedBindings->countSuggestedBindings <= 0
+		|| SuggestedBindings->suggestedBindings == nullptr)
+	{
+		return XrResult::XR_ERROR_VALIDATION_FAILURE;
+	}	
+
+	UE_LOG(LogHMD, Verbose, TEXT("XrSuggestInteractionProfileBindings profile: %s"), ANSI_TO_TCHAR(PathToString(SuggestedBindings->interactionProfile)));
+
+	CreateInteractionProfiles();
+
+
+	const TMap<XrPath, EOXRVisionOSControllerButton>* InteractionProfilePairs = InteractionProfiles.Find(SuggestedBindings->interactionProfile);
+	if (InteractionProfilePairs == nullptr)
+	{
+		// We are successfully ignoring this interaction profile.
+		UE_LOG(LogOXRVisionOS, Verbose, TEXT("  Ignoring %s because OXRVisionOS does not support it"), ANSI_TO_TCHAR(PathToString(SuggestedBindings->interactionProfile)));
+		return XrResult::XR_SUCCESS;
+	}
+
+	FInteractionProfileBinding* InteractionProfileBinding = InteractionProfileBindings.Find(SuggestedBindings->interactionProfile);
+	if (InteractionProfileBinding)
+	{
+		// Previous bindings are overwritten.
+		InteractionProfileBinding->Bindings.Empty(SuggestedBindings->countSuggestedBindings);
+	}
+	else
+	{
+		InteractionProfileBinding = &InteractionProfileBindings.Add(SuggestedBindings->interactionProfile);
+		InteractionProfileBinding->InteractionProfilePath = SuggestedBindings->interactionProfile;
+		InteractionProfileBinding->Bindings.Reserve(SuggestedBindings->countSuggestedBindings);
+	}
+
+	const uint32_t SuggestedBindingCount = SuggestedBindings->countSuggestedBindings;
+	for (uint32_t i = 0; i < SuggestedBindingCount; i++)
+	{
+		const XrActionSuggestedBinding& Binding = SuggestedBindings->suggestedBindings[i];
+		// Map to the native button
+		const EOXRVisionOSControllerButton* Source = InteractionProfilePairs->Find(Binding.binding);
+		if (Source == nullptr)
+		{
+			UE_LOG(LogOXRVisionOS, Warning, TEXT("Unsupported path %s in profile %s"), ANSI_TO_TCHAR(PathToString(Binding.binding)), ANSI_TO_TCHAR(PathToString(SuggestedBindings->interactionProfile)));
+			return XrResult::XR_ERROR_PATH_UNSUPPORTED;
+		}
+		else
+		{
+			EOXRVisionOSControllerButton ButtonSource = *Source;
+			FOXRVisionOSAction* Action = (FOXRVisionOSAction*)Binding.action;
+			if (Action->GetActionSet().IsAttached())
+			{
+				return XrResult::XR_ERROR_ACTIONSETS_ALREADY_ATTACHED;
+			}
+			InteractionProfileBinding->Bindings.Emplace(Binding.action, Binding.binding, ButtonSource);
+
+			// Update TopLevelPathInteractionProfileMap as needed
+			const XrPath& TopLevelPath = InputPathToTopLevelPathMap.FindChecked(Binding.binding);
+			const XrPath& CurrentTopLevelPathInteractionProfile = TopLevelPathInteractionProfileMap.FindChecked(TopLevelPath);
+
+			if (IsFirstInputProfilePreferred(SuggestedBindings->interactionProfile, CurrentTopLevelPathInteractionProfile))
+			{
+				UE_LOG(LogOXRVisionOS, Log, TEXT("TopLevelPath %s set to InteractionProfile %s"), ANSI_TO_TCHAR(PathToString(TopLevelPath)), ANSI_TO_TCHAR(PathToString(SuggestedBindings->interactionProfile)));
+				TopLevelPathInteractionProfileMap[TopLevelPath] = SuggestedBindings->interactionProfile;
+			}
+		}
+	}
+
+	return XrResult::XR_SUCCESS;
+}
+
+XrResult FOXRVisionOSInstance::XrPollEvent(
+	XrEventDataBuffer* EventData)
+{
+	// Get the next event buffer
+	TSharedPtr<IEventDataHolder, ESPMode::ThreadSafe> EventHolder;
+	if (EventQueue.Dequeue(EventHolder) == false)
+	{
+		return XR_EVENT_UNAVAILABLE;
+	}
+
+	EventHolder->CopyInto(EventData);
+	return XR_SUCCESS;
+}
diff --git a/Engine/Platforms/VisionOS/Plugins/Runtime/OpenXRVisionOS/Source/OXRVisionOS/Private/OXRVisionOSInstance.h b/Engine/Platforms/VisionOS/Plugins/Runtime/OpenXRVisionOS/Source/OXRVisionOS/Private/OXRVisionOSInstance.h
new file mode 100644
--- /dev/null
+++ b/Engine/Platforms/VisionOS/Plugins/Runtime/OpenXRVisionOS/Source/OXRVisionOS/Private/OXRVisionOSInstance.h
@@ -0,0 +1,247 @@
+// Copyright Epic Games, Inc. All Rights Reserved.
+
+#pragma once
+
+#include "CoreMinimal.h"
+#include "Engine/EngineBaseTypes.h"
+#include "OXRVisionOS.h"
+#include "Containers/Queue.h"
+#include "Containers/Map.h"
+#include "Delegates/IDelegateInstance.h"
+
+class FOXRVisionOS;
+class FOXRVisionOSSession;
+class FOXRVisionOSActionSet;
+
+enum class EOXRVisionOSControllerButton : int32;
+
+class FOXRVisionOSInstance
+{
+public:
+	static XrResult Create(TSharedPtr<FOXRVisionOSInstance, ESPMode::ThreadSafe>& OutInstance, const XrInstanceCreateInfo* createInfo, FOXRVisionOS* InModule);
+	FOXRVisionOSInstance(const XrInstanceCreateInfo* createInfo, FOXRVisionOS* Module);
+
+	~FOXRVisionOSInstance();
+
+	XrResult XrGetInstanceProperties(
+		XrInstanceProperties*						instanceProperties);
+
+	XrResult XrGetSystemProperties(
+		XrSystemId                                  systemId,
+		XrSystemProperties* properties);
+
+	XrResult XrEnumerateEnvironmentBlendModes(
+		XrSystemId                                  systemId,
+		XrViewConfigurationType                     viewConfigurationType,
+		uint32_t                                    environmentBlendModeCapacityInput,
+		uint32_t*									environmentBlendModeCountOutput,
+		XrEnvironmentBlendMode*						environmentBlendModes);
+
+	XrResult XrEnumerateViewConfigurations(
+		XrSystemId                                  systemId,
+		uint32_t                                    viewConfigurationTypeCapacityInput,
+		uint32_t*									viewConfigurationTypeCountOutput,
+		XrViewConfigurationType*					viewConfigurationTypes);
+
+	XrResult XrGetViewConfigurationProperties(
+		XrSystemId                                  systemId,
+		XrViewConfigurationType                     viewConfigurationType,
+		XrViewConfigurationProperties*				configurationProperties);
+
+	XrResult XrEnumerateViewConfigurationViews(
+		XrSystemId                                  systemId,
+		XrViewConfigurationType                     viewConfigurationType,
+		uint32_t                                    viewCapacityInput,
+		uint32_t* viewCountOutput,
+		XrViewConfigurationView* views);
+
+	XrResult XrGetSystem(
+		const XrSystemGetInfo*						getInfo,
+		XrSystemId*									systemId);
+
+	bool IsHMDAvailable();
+
+	XrResult XrCreateSession(
+		const XrSessionCreateInfo*					createInfo,
+		XrSession*									session);
+
+	XrResult DestroySession() { Session = nullptr; return XR_SUCCESS; }
+
+	static bool CheckSession(XrSession session) { return (session != XR_NULL_HANDLE) && (session == (XrSession)Session.Get()); }
+
+	XrResult XrStringToPath(
+		const char* pathString,
+		XrPath* path);
+
+	XrResult XrPathToString(
+		XrPath                                      path,
+		uint32_t                                    bufferCapacityInput,
+		uint32_t*									bufferCountOutput,
+		char*										buffer);
+
+	char* PathToString(XrPath Path) const
+	{
+		return XrPaths[Path].Get();
+	}
+
+	XrResult XrCreateActionSet(
+		const XrActionSetCreateInfo*				createInfo,
+		XrActionSet*								actionSet);
+
+	XrResult DestroyActionSet(FOXRVisionOSActionSet* ActionSet);
+
+	XrResult XrSuggestInteractionProfileBindings(
+		const XrInteractionProfileSuggestedBinding* suggestedBindings);
+
+	XrResult XrPollEvent(
+		XrEventDataBuffer*							eventData);
+
+	template <typename T1>
+	void EnqueueEvent(T1& InEvent)
+	{
+		TSharedPtr<IEventDataHolder, ESPMode::ThreadSafe> DataHolderPtr(new FEventDataHolder<T1>(InEvent));
+		EventQueue.Enqueue(DataHolderPtr);
+	}
+
+	bool IsExtensionEnabled(const ANSICHAR* ExtensionName) const { return EnabledExtensions.Contains(ExtensionName); } 
+
+private:
+	bool bCreateFailed = false;
+	FOXRVisionOS* Module = nullptr;
+	static TSharedPtr<FOXRVisionOSSession, ESPMode::ThreadSafe> Session;
+
+	void OnWorldTickStart(UWorld* World, ELevelTick TickType, float DeltaTime);
+	FDelegateHandle OnWorldTickStartDelegateHandle;
+	bool bHMDWorn = false;
+
+	struct FAnsiKeyFunc : BaseKeyFuncs<const ANSICHAR*, const ANSICHAR*, false>
+	{
+		typedef typename TTypeTraits<const ANSICHAR*>::ConstPointerType KeyInitType;
+		typedef typename TCallTraits<const ANSICHAR*>::ParamType ElementInitType;
+
+		/**
+		 * @return The key used to index the given element.
+		 */
+		static FORCEINLINE KeyInitType GetSetKey(ElementInitType Element)
+		{
+			return Element;
+		}
+
+		/**
+		 * @return True if the keys match.
+		 */
+		static FORCEINLINE bool Matches(KeyInitType A, KeyInitType B)
+		{
+			return TCString<ANSICHAR>::Stricmp(A, B) == 0;
+		}
+
+		/** Calculates a hash index for a key. */
+		static FORCEINLINE uint32 GetKeyHash(KeyInitType Key)
+		{
+			return FCrc::Strihash_DEPRECATED(Key);
+		}
+	};
+	TSet<const ANSICHAR*, FAnsiKeyFunc> EnabledExtensions;
+
+	XrInstanceProperties InstanceProperties;
+	XrSystemProperties SystemProperties;
+//	XrSystemEyeGazeInteractionPropertiesEXT EyeGazeInteractionSystemProperties;
+
+	struct IEventDataHolder
+	{
+		virtual ~IEventDataHolder() {}
+		virtual void CopyInto(XrEventDataBuffer* eventData) = 0;
+	};
+
+	template <typename T1>
+	struct FEventDataHolder : public IEventDataHolder
+	{
+		FEventDataHolder(T1 InEvent)	{ Event = InEvent; check(Event.next == nullptr); }
+		virtual ~FEventDataHolder()		{};
+		virtual void CopyInto(XrEventDataBuffer* EventData) override
+		{
+			check(sizeof(T1) <= sizeof(XrEventDataBuffer))
+			*reinterpret_cast<T1*>(EventData) = Event;
+		}
+	private:
+		T1 Event;
+	};
+
+	TQueue<TSharedPtr<IEventDataHolder, ESPMode::ThreadSafe>, EQueueMode::Mpsc> EventQueue;
+
+	XrPath StringToPath(const char* pathString);
+	XrPath StringToPath(const char* pathString, uint32 PathSize);
+	TArray<TUniquePtr<ANSICHAR[]>> XrPaths; //An XrPath is an index into this array
+	struct FStringToXrPathKeyFunc
+	{
+		typedef const ANSICHAR* KeyType;
+		typedef const ANSICHAR* KeyInitType;
+		typedef const TPairInitializer<const ANSICHAR*, XrPath>& ElementInitType;
+
+		enum { bAllowDuplicateKeys = false };
+
+		static FORCEINLINE bool Matches(const ANSICHAR* A, const ANSICHAR* B)
+		{
+			return TCString<ANSICHAR>::Stricmp(A, B) == 0;
+		}
+
+		static FORCEINLINE uint32 GetKeyHash(const ANSICHAR* Key)
+		{
+			return FCrc::Strihash_DEPRECATED(Key);
+		}
+
+		static FORCEINLINE KeyInitType GetSetKey(ElementInitType Element)
+		{
+			return Element.Key;
+		}
+	};
+	TMap<const ANSICHAR*, XrPath, FDefaultSetAllocator, FStringToXrPathKeyFunc> StringToXrPathMap;
+	TArray<TSharedPtr<FOXRVisionOSActionSet, ESPMode::ThreadSafe>> ActionSets;
+public:
+	TSharedPtr<FOXRVisionOSActionSet, ESPMode::ThreadSafe> GetActionSetSharedPtr(XrActionSet ActionSet);
+private:
+
+
+	// Interaction profiles define a set of 'input paths' and how they should map to runtime hardware, like controllers.  
+	// We have a set of source buttons that are 1:1 with specific hardware.
+	// The OXRVisionOS native interaction profile, therefore, simply maps every path to its button.
+	// Other profiles will map their input paths to one button.
+	// Interaction profiles will, currently, never bind multiple source buttons to a single profile path.
+	TMap<XrPath, TMap<XrPath, EOXRVisionOSControllerButton>> InteractionProfiles; // Profile, input path, native path.
+	void CreateInteractionProfiles();
+
+	// The runtime must select one interaction profile for each top level path from among those the application attempts to bind.
+	bool IsFirstInputProfilePreferred(XrPath First, XrPath Second) const;
+	TMap<XrPath, XrPath> TopLevelPathInteractionProfileMap; // Map of top level paths to most preferred bound input profile
+	TMap<XrPath, XrPath> InputPathToTopLevelPathMap; // Map of input paths to Top level paths (eg /user/hand/left)
+
+	// XrSuggestInteractionProfileBindings takes a set of action-binding mappings and one of the InteractionProfiles we define and generates
+	// an array of bindings for that profile.  An action can have multiple suggested paths.
+	// See XrAttachSessionActionSets for the next step in the input process.
+public:
+	struct FInteractionProfileBinding
+	{
+		XrPath InteractionProfilePath = XR_NULL_PATH;
+		struct FBinding
+		{
+			FBinding()
+			{}
+			FBinding(XrAction InAction, XrPath InPath, EOXRVisionOSControllerButton InButton)
+				: Action(InAction)
+				, Path(InPath)
+				, Button(InButton)
+			{}
+			XrAction Action;
+			XrPath Path = XR_NULL_PATH;
+			EOXRVisionOSControllerButton Button;
+		};
+		TArray<FBinding> Bindings;
+	};
+	TMap<XrPath, FInteractionProfileBinding> InteractionProfileBindings;
+	TArray<XrPath> PreferredInteractionProfileBindingList; // The list of supported interaction profiles ordered from most to least preferred.
+	const TMap<XrPath, FInteractionProfileBinding>& GetInteractionProfileBindings() const {	return InteractionProfileBindings; }
+	const TArray<XrPath>& GetPreferredInteractionProfileBindingList() const { return PreferredInteractionProfileBindingList; }
+	XrPath GetCurrentInteractionProfile(XrPath TopLevelUserPath) const { return TopLevelPathInteractionProfileMap.FindChecked(TopLevelUserPath); }
+	XrPath InputPathToTopLevelPath(XrPath InPath) const { return InputPathToTopLevelPathMap.Contains(InPath) ? InputPathToTopLevelPathMap.FindChecked(InPath) : XR_NULL_PATH; }
+	XrPath TopLevelPathToInteractionProfile(XrPath InPath) const { return TopLevelPathInteractionProfileMap.Contains(InPath) ? TopLevelPathInteractionProfileMap.FindChecked(InPath) : XR_NULL_PATH; }
+};
diff --git a/Engine/Platforms/VisionOS/Plugins/Runtime/OpenXRVisionOS/Source/OXRVisionOS/Private/OXRVisionOSPlatformRHI.h b/Engine/Platforms/VisionOS/Plugins/Runtime/OpenXRVisionOS/Source/OXRVisionOS/Private/OXRVisionOSPlatformRHI.h
new file mode 100644
--- /dev/null
+++ b/Engine/Platforms/VisionOS/Plugins/Runtime/OpenXRVisionOS/Source/OXRVisionOS/Private/OXRVisionOSPlatformRHI.h
@@ -0,0 +1,12 @@
+// Copyright Epic Games, Inc. All Rights Reserved.
+
+#pragma once
+
+#define XR_USE_GRAPHICS_API_OXRVISIONOS
+
+#ifdef XR_USE_GRAPHICS_API_OXRVISIONOS
+#include "DynamicRHI.h"
+#endif // XR_USE_GRAPHICS_API_OXRVISIONOS
+
+#include <openxr/openxr_platform.h>
+
diff --git a/Engine/Platforms/VisionOS/Plugins/Runtime/OpenXRVisionOS/Source/OXRVisionOS/Private/OXRVisionOSPlatformUtils.h b/Engine/Platforms/VisionOS/Plugins/Runtime/OpenXRVisionOS/Source/OXRVisionOS/Private/OXRVisionOSPlatformUtils.h
new file mode 100644
--- /dev/null
+++ b/Engine/Platforms/VisionOS/Plugins/Runtime/OpenXRVisionOS/Source/OXRVisionOS/Private/OXRVisionOSPlatformUtils.h
@@ -0,0 +1,236 @@
+// Copyright Epic Games, Inc. All Rights Reserved.
+
+#pragma once
+
+#include "CoreMinimal.h"
+#include "OpenXRCore.h"
+#include "Epic_openxr.h"
+#import <simd/simd.h>
+
+// OXRVisionOS works with UE4, OpenXR, and VisionOS types.  OpenXR and VisionOS coordinate systems are the same.
+
+namespace OXRVisionOS
+{
+	struct TrackerPoseData
+	{
+		simd_float3 position = {0.0f,0.0f,0.0f};
+		simd_float4 orientation = {0.0f,0.0f,0.0f,1.0f};
+	};
+
+	struct TrackerResultData
+	{
+		TrackerPoseData Pose;
+	};
+
+
+	const XrVector3f XrZeroVector = {};
+
+	FORCEINLINE FVector ToFVector(const simd_float3& InVector)
+	{
+		return FVector(
+			-InVector.z,
+			InVector.x,
+			InVector.y
+		);
+	}
+	FORCEINLINE simd_float3 ToSimdFVector3(const FVector& InVector)
+	{
+		return simd_float3{
+			static_cast<float>(InVector.Y),
+			static_cast<float>(InVector.Z),
+			static_cast<float>(-InVector.X)
+		};
+	}
+	FORCEINLINE simd_float3 ToSimdFVector3(const XrVector3f& InVector)
+	{
+		return simd_float3{
+			InVector.x,
+			InVector.y,
+			InVector.z
+		};
+	}
+	FORCEINLINE XrVector3f ToXrVector3f(const FVector& InVector)
+	{
+		return XrVector3f{
+			static_cast<float>(InVector.Y),
+			static_cast<float>(InVector.Z),
+			static_cast<float>(-InVector.X)
+		};
+	}
+
+	FORCEINLINE XrVector3f ToXrVector3f(const simd_float3& InVector)
+	{
+		return XrVector3f{
+			InVector.x,
+			InVector.y,
+			InVector.z
+		};
+	}
+
+	FORCEINLINE FQuat ToFQuat(const simd_float4& InQuaternion)
+	{
+		return FQuat(
+			-InQuaternion.z,
+			InQuaternion.x,
+			InQuaternion.y,
+			-InQuaternion.w
+		);
+	}
+	FORCEINLINE simd_float4 ToSimdFQuaternion(const FQuat& InQuaternion)
+	{
+		return simd_float4{
+			static_cast<float>(InQuaternion.Y),
+			static_cast<float>(InQuaternion.Z),
+			static_cast<float>(-InQuaternion.X),
+			static_cast<float>(-InQuaternion.W)
+		};
+	}
+	FORCEINLINE simd_float4 ToSimdFQuaternion(const XrQuaternionf& InQuaternion)
+	{
+		return simd_float4{
+			InQuaternion.x,
+			InQuaternion.y,
+			InQuaternion.z,
+			InQuaternion.w
+		};
+	}
+	FORCEINLINE XrQuaternionf ToXrQuaternionf(const FQuat& InQuaternion)
+	{
+		return XrQuaternionf{
+			static_cast<float>(InQuaternion.Y),
+			static_cast<float>(InQuaternion.Z),
+			static_cast<float>(-InQuaternion.X),
+			static_cast<float>(-InQuaternion.W)
+		};
+	}
+
+    /**
+     * Convert's an ARKit 'Y up' 'right handed' coordinate system transform to Unreal's 'Z up'
+     * 'left handed' coordinate system.
+     *
+     * Scale is removed because some ARKit transforms are less normalized than unreal requires.
+     */
+    static FORCEINLINE FTransform ToFTransform(const matrix_float4x4& RawYUpMatrix)
+    {
+        // Conversion here is as per SteamVRHMD::ToFMatrix
+        FMatrix RawYUpFMatrix(
+            FPlane(RawYUpMatrix.columns[0][0], RawYUpMatrix.columns[0][1], RawYUpMatrix.columns[0][2], RawYUpMatrix.columns[0][3]),
+            FPlane(RawYUpMatrix.columns[1][0], RawYUpMatrix.columns[1][1], RawYUpMatrix.columns[1][2], RawYUpMatrix.columns[1][3]),
+            FPlane(RawYUpMatrix.columns[2][0], RawYUpMatrix.columns[2][1], RawYUpMatrix.columns[2][2], RawYUpMatrix.columns[2][3]),
+            FPlane(RawYUpMatrix.columns[3][0], RawYUpMatrix.columns[3][1], RawYUpMatrix.columns[3][2], RawYUpMatrix.columns[3][3]));
+
+        // ARKit, particularly when restoring a previously saved AR world map, sometimes has transforms who's normalization tolerance is not
+        // tight enough to avoid ensures or even failure that result in identity rotations from FQuat.  So we remove scaling here.
+        RawYUpFMatrix.RemoveScaling(UE_KINDA_SMALL_NUMBER / 100.0f);
+
+        // Extract & convert translation
+        FVector Translation = FVector( -RawYUpFMatrix.M[3][2], RawYUpFMatrix.M[3][0], RawYUpFMatrix.M[3][1] );
+
+        // Extract & convert rotation
+        FQuat RawRotation( RawYUpFMatrix );
+        FQuat Rotation( -RawRotation.Z, RawRotation.X, RawRotation.Y, -RawRotation.W );
+
+        return FTransform( Rotation, Translation );
+    }
+
+	FORCEINLINE FTransform ToFTransform(const TrackerPoseData& InPoseData)
+	{
+		return FTransform(
+			ToFQuat(InPoseData.orientation),
+			ToFVector(InPoseData.position)
+			);
+	}
+	FORCEINLINE TrackerPoseData ToPoseData(const FTransform& InTransform)
+	{
+		TrackerPoseData PoseData;
+		PoseData.orientation = ToSimdFQuaternion(InTransform.GetRotation());
+		PoseData.position = ToSimdFVector3(InTransform.GetLocation());
+		return PoseData;
+	}
+
+	FORCEINLINE XrPosef ToXrPose(const TrackerPoseData& InPoseData)
+	{
+		XrPosef Ret;
+		Ret.orientation = XrQuaternionf{
+			InPoseData.orientation.x,
+			InPoseData.orientation.y,
+			InPoseData.orientation.z,
+			InPoseData.orientation.w };
+		Ret.position = XrVector3f{ 
+			InPoseData.position.x,
+			InPoseData.position.y,
+			InPoseData.position.z };
+		return Ret;
+	}
+    FORCEINLINE XrPosef ToXrPose(const simd_float4x4& RawYUpMatrix)
+    {
+        // Conversion here is as per SteamVRHMD::ToFMatrix
+        FMatrix RawYUpFMatrix(
+            FPlane(RawYUpMatrix.columns[0][0], RawYUpMatrix.columns[0][1], RawYUpMatrix.columns[0][2], RawYUpMatrix.columns[0][3]),
+            FPlane(RawYUpMatrix.columns[1][0], RawYUpMatrix.columns[1][1], RawYUpMatrix.columns[1][2], RawYUpMatrix.columns[1][3]),
+            FPlane(RawYUpMatrix.columns[2][0], RawYUpMatrix.columns[2][1], RawYUpMatrix.columns[2][2], RawYUpMatrix.columns[2][3]),
+            FPlane(RawYUpMatrix.columns[3][0], RawYUpMatrix.columns[3][1], RawYUpMatrix.columns[3][2], RawYUpMatrix.columns[3][3]));
+
+        // ARKit, particularly when restoring a previously saved AR world map, sometimes has transforms who's normalization tolerance is not
+        // tight enough to avoid ensures or even failure that result in identity rotations from FQuat.  So we remove scaling here.
+        RawYUpFMatrix.RemoveScaling(UE_KINDA_SMALL_NUMBER / 100.0f);
+
+        // Extract rotation
+        const FQuat RawRotation( RawYUpFMatrix );
+
+        XrPosef Ret;
+        Ret.orientation.x = RawRotation.X;
+        Ret.orientation.y = RawRotation.Y;
+        Ret.orientation.z = RawRotation.Z;
+        Ret.orientation.w = RawRotation.W;
+        Ret.position.x = RawYUpFMatrix.M[3][0];
+        Ret.position.y = RawYUpFMatrix.M[3][1];
+        Ret.position.z = RawYUpFMatrix.M[3][2];
+        return Ret;
+    }
+	FORCEINLINE XrPosef ToXrPose(const FTransform& InTransform)
+	{
+		XrPosef Ret;
+		Ret.orientation = ToXrQuaternionf(InTransform.GetRotation());
+		Ret.position = ToXrVector3f(InTransform.GetLocation());
+		return Ret;
+	}
+
+	FORCEINLINE XrSpaceLocationFlags ToXrSpaceLocationFlags(TrackerResultData& InResultData)
+	{
+//TODO finish this
+		XrSpaceLocationFlags OutLocationFlags = 0;
+		 	OutLocationFlags |= XR_SPACE_LOCATION_POSITION_VALID_BIT;
+		// 	OutLocationFlags |= (XR_SPACE_LOCATION_POSITION_VALID_BIT | XR_SPACE_LOCATION_POSITION_TRACKED_BIT);
+
+
+		 	OutLocationFlags |= XR_SPACE_LOCATION_ORIENTATION_VALID_BIT;
+		// 	OutLocationFlags |= (XR_SPACE_LOCATION_ORIENTATION_VALID_BIT | XR_SPACE_LOCATION_ORIENTATION_TRACKED_BIT);
+
+
+		return OutLocationFlags;
+	}
+
+	FORCEINLINE XrSpaceVelocityFlags ToXrSpaceVelocityFlags(TrackerResultData& InResultData)
+	{
+//TODO finish this
+		XrSpaceVelocityFlags OutFlags = 0;
+
+		// 	OutFlags |= XR_SPACE_VELOCITY_LINEAR_VALID_BIT;
+		// 	OutFlags |= XR_SPACE_VELOCITY_ANGULAR_VALID_BIT;
+
+		return OutFlags;
+	}
+	
+
+	// CFTimeInterval is a double of seconds.  XrTime is uint64 nanoseconds.
+	FORCEINLINE XrTime APITimeToXrTime(CFTimeInterval TimeInterval)
+	{
+		return (XrTime)(TimeInterval * 1000000000.0);
+	}
+
+	FORCEINLINE CFTimeInterval XrTimeToAPITime(XrTime xrTime)
+	{
+		return xrTime / 1000000000.0;
+	}
+};
diff --git a/Engine/Platforms/VisionOS/Plugins/Runtime/OpenXRVisionOS/Source/OXRVisionOS/Private/OXRVisionOSSession.cpp b/Engine/Platforms/VisionOS/Plugins/Runtime/OpenXRVisionOS/Source/OXRVisionOS/Private/OXRVisionOSSession.cpp
new file mode 100644
--- /dev/null
+++ b/Engine/Platforms/VisionOS/Plugins/Runtime/OpenXRVisionOS/Source/OXRVisionOS/Private/OXRVisionOSSession.cpp
@@ -0,0 +1,1422 @@
+// Copyright Epic Games, Inc. All Rights Reserved.
+
+#include "OXRVisionOSSession.h"
+#include "OXRVisionOSAction.h"
+#include "OXRVisionOSActionSet.h"
+#include "OXRVisionOSController.h"
+#include "OXRVisionOSInstance.h"
+#include "OXRVisionOSPlatformUtils.h"
+#include "OXRVisionOS_RenderBridge.h"
+#include "OXRVisionOSRuntimeSettings.h"
+#include "OXRVisionOSSettingsTypes.h"
+#include "OXRVisionOSSpace.h"
+#include "OXRVisionOSSwapchain.h"
+#include "OXRVisionOSTracker.h"
+#include "OXRVisionOS_openxr_platform.h"
+#include "OXRVisionOSAppDelegate.h"
+#include "PixelFormat.h"
+#include "StereoRenderTargetManager.h"
+#include "HAL/RunnableThread.h"
+
+#include "RenderingThread.h"
+
+#include "CoreGlobals.h"
+#include "RHICommandList.h"
+#include "MetalRHIVisionOSBridge.h"
+
+#define CHECK_API_RET_CREATE_FAIL(___function) {int __ret = RESULT_OK; __ret = (___function); if(__ret < RESULT_OK) { UE_LOG(LogOXRVisionOS, Warning, TEXT("Call to " #___function " failed, returned %d 0x%x"), __ret, __ret); UE_DEBUG_BREAK(); bCreateFailed = true; return; }}
+#define CHECK_API_RET_RETURN_FAIL(___function) {int __ret = RESULT_OK; __ret = (___function); if(__ret < RESULT_OK) { UE_LOG(LogOXRVisionOS, Warning, TEXT("Call to " #___function " failed, returned %d 0x%x"), __ret, __ret); UE_DEBUG_BREAK(); return XrResult::XR_ERROR_RUNTIME_FAILURE; }}
+#define CHECK_API_RET_WARN(___function) {int __ret = RESULT_OK; __ret = (___function); if(__ret < RESULT_OK) { UE_LOG(LogOXRVisionOS, Warning, TEXT("Call to " #___function " failed, returned %d 0x%x"), __ret, __ret);}}
+
+#define OXRVISIONOS_REPROJECTION_ANALYSIS 0
+
+// Helper macro for checking Api error codes. If the result of a function call is negative, causes a fatal log.
+#define FA_CHECK_F(code, format, ...) do { int32 Result; { Result = (code); } if (Result < 0) { UE_LOG(LogOXRVisionOS, Fatal, TEXT("%s failed with error code: 0x%08x") format, TEXT(#code), Result, ##__VA_ARGS__); } } while(false)
+#define FA_CHECK(code) FA_CHECK_F(code, TEXT(""))
+
+namespace OXRVisionOSCVars
+{
+	static TAutoConsoleVariable<bool> CVarDisableReprojection(
+		TEXT("vr.OXRVisionOS.DisableReprojection"),
+		false,
+		TEXT("Disables reprojection for debugging purposes"));
+}
+
+namespace OXRVisionOSSessionHelpers
+{
+	enum ELayerIndex : uint8
+	{
+		Background = 0,
+		//CompositedOverlays,
+		Num
+	};
+
+	constexpr uint32 OXRVisionOSBackbufferLength = 2;
+	constexpr uint32 MaxOXRVisionOSLayers = 1; //TODO for now I only want the one stereo layer, 
+}
+
+XrResult FOXRVisionOSSession::Create(TSharedPtr<FOXRVisionOSSession, ESPMode::ThreadSafe>& OutSession, const XrSessionCreateInfo* createInfo, FOXRVisionOSInstance* Instance)
+{
+	if (createInfo == nullptr || Instance == nullptr)
+	{
+		return XrResult::XR_ERROR_VALIDATION_FAILURE;
+	}
+
+	if (createInfo->type != XR_TYPE_SESSION_CREATE_INFO)
+	{
+		return XrResult::XR_ERROR_VALIDATION_FAILURE;
+	}
+
+	const XrGraphicsBindingOXRVisionOSEPIC* GraphicsBinding = OpenXR::FindChainedStructByType<XrGraphicsBindingOXRVisionOSEPIC>(createInfo->next, (XrStructureType)XR_TYPE_GRAPHICS_BINDING_OXRVISIONOS_EPIC);
+	if (GraphicsBinding)
+	{
+		//TODO use the graphics binding?
+	}
+
+	OutSession = MakeShared<FOXRVisionOSSession, ESPMode::ThreadSafe>(createInfo, Instance);
+	if (OutSession->bCreateFailed)
+	{
+		OutSession = nullptr;
+		return XrResult::XR_ERROR_RUNTIME_FAILURE;
+	}
+	return XrResult::XR_SUCCESS;
+}
+
+
+FOXRVisionOSSession::FOXRVisionOSSession(const XrSessionCreateInfo* createInfo, FOXRVisionOSInstance* InInstance)
+{
+	Instance = InInstance;
+
+	if (FOXRVisionOSTracker::Create(Tracker, Instance->IsExtensionEnabled(XR_EXT_EYE_GAZE_INTERACTION_EXTENSION_NAME)) == false)
+	{
+		UE_LOG(LogOXRVisionOS, Warning, TEXT("FOXRVisionOSTracker::Create failed.  See logging above.")); 
+		bCreateFailed = true;
+		return;
+	}
+
+	//if (FOXRVisionOSController::Create(Controllers, this) == false)
+	//{
+	//	UE_LOG(LogOXRVisionOS, Warning, TEXT("FOXRVisionOSController::Create failed.  See logging above."));
+	//	bCreateFailed = true;
+	//	return;
+	//}
+
+	const UOXRVisionOSRuntimeSettings* Settings = GetDefault<UOXRVisionOSRuntimeSettings>();
+	AimPoseAdjustment = FQuat(Settings->OXRVisionOSAimPoseAdjustment);
+	GripPoseAdjustment = FQuat(Settings->OXRVisionOSGripPoseAdjustment);
+
+	static_assert(OXRVisionOSSessionHelpers::ELayerIndex::Num <= OXRVisionOSSessionHelpers::MaxOXRVisionOSLayers);
+
+	RenderBridge = FOXRVisionOS::Get()->GetRenderBridge();
+	check(RenderBridge != nullptr);
+
+	ARKitSession = ar_session_create();
+
+	XrWaitFrameEvent = FPlatformProcess::GetSynchEventFromPool(false);
+	XrWaitFrameEvent->Trigger();
+
+	SetSessionState(XR_SESSION_STATE_IDLE);
+	
+	if (Instance->IsHMDAvailable())
+	{
+		SetSessionState(XR_SESSION_STATE_READY);
+	}
+	else
+	{
+	//	// else we need to wait for the hmd to become available, not relevant on this all-in-one hmd device
+	//	OnWorldTickStartDelegateHandle = FWorldDelegates::OnWorldTickStart.AddRaw(this, &FOXRVisionOSSession::OnWorldTickStart);
+	}
+}
+
+void FOXRVisionOSSession::OnWorldTickStart(UWorld* World, ELevelTick TickType, float DeltaTime)
+{
+	if (SessionState == XR_SESSION_STATE_IDLE)
+	{
+		if (Instance->IsHMDAvailable())
+		{
+			SetSessionState(XR_SESSION_STATE_READY);
+		}
+	}
+
+	// Stop trying to transition to READY when no longer IDLE.
+	if (SessionState != XR_SESSION_STATE_IDLE)
+	{
+		FWorldDelegates::OnWorldTickStart.Remove(OnWorldTickStartDelegateHandle);
+	}
+}
+
+FOXRVisionOSSession::~FOXRVisionOSSession()
+{
+	FWorldDelegates::OnWorldTickStart.Remove(OnWorldTickStartDelegateHandle);
+
+	//XrDestroySession can be called from any session state, so all of them must cleanup correctly.
+	if (IsSessionRunning())
+	{
+		EndSessionInternal();
+	}
+
+	Controllers = nullptr;
+	Tracker = nullptr;
+
+	FPlatformProcess::ReturnSynchEventToPool(XrWaitFrameEvent);
+	XrWaitFrameEvent = nullptr;
+
+	if (bCreateFailed)
+	{
+		UE_LOG(LogOXRVisionOS, Warning, TEXT("Destructing FOXRVisionOSSession because session create failed.  You may see more warnings as api shutdown calls fail, depending on how the create failed."));
+	}
+
+	{
+		// We need to unregister the swapchain buffers from the video out.  
+		// First we have to make sure they aren't being used by flipping to blank and flushing it through.
+		// Then we can unregister.
+
+		// Completely drain all remaining GPU work	
+		ENQUEUE_RENDER_COMMAND(FlipToBlank)([this](FRHICommandListImmediate& RHICmdList)
+			{
+				// Flip to blank and drain the GPU
+				RHICmdList.EnqueueLambda([this](FRHICommandListImmediate& RHICmdList)
+					{
+						//TODO
+						//MetalRHIOXRVisionOS::SetFlipToBlankAndFinalize(RHICmdList);
+
+						//// Wait for all outstanding flips to clear, including our flip to blank.
+					});
+
+				// Perform a full flush
+				RHICmdList.ImmediateFlush(EImmediateFlushType::FlushRHIThreadFlushResources);
+
+				RHICmdList.BlockUntilGPUIdle();
+			});
+
+		FlushRenderingCommands();
+	}
+}
+
+XrResult FOXRVisionOSSession::XrDestroySession()
+{
+	// Instance->DestroySession() can delete this, so better just return after that.
+	return Instance->DestroySession();
+}
+
+XrResult FOXRVisionOSSession::XrEnumerateReferenceSpaces(
+	uint32_t                                    SpaceCapacityInput,
+	uint32_t*									SpaceCountOutput,
+	XrReferenceSpaceType*						OutSpaces)
+{
+	if (SpaceCountOutput == nullptr)
+	{
+		return XrResult::XR_ERROR_VALIDATION_FAILURE;
+	}
+	if (SpaceCapacityInput != 0 && OutSpaces == nullptr)
+	{
+		return XrResult::XR_ERROR_VALIDATION_FAILURE;
+	}
+
+	const int NumSupportedSpaceTypes = 2;//3;
+
+	*SpaceCountOutput = NumSupportedSpaceTypes;
+	if (SpaceCapacityInput != 0)
+	{
+		if (SpaceCapacityInput < *SpaceCountOutput)
+		{
+			return XrResult::XR_ERROR_SIZE_INSUFFICIENT;
+		}
+
+		OutSpaces[0] = XR_REFERENCE_SPACE_TYPE_VIEW;
+		OutSpaces[1] = XR_REFERENCE_SPACE_TYPE_LOCAL;
+		//OutSpaces[2] = XR_REFERENCE_SPACE_TYPE_STAGE;
+
+		//CoordinateSystemType CurrentCoordinateSystem = COORDINATE_SYSTEM_LOCAL_FLOOR;
+		// const int32 Ret = SetCoordinateSystem(COORDINATE_SYSTEM_STAGE);
+		// if (Ret == RESULT_OK)
+		// {
+		// 	UE_LOG(LogOXRVisionOS, Log, TEXT("XrEnumerateReferenceSpaces called SetCoordinateSystem(COORDINATE_SYSTEM_STAGE)."));
+		// }
+		// else
+		// {
+		// 	UE_LOG(LogOXRVisionOS, Warning, TEXT("XrEnumerateReferenceSpaces call to SetCoordinateSystem(COORDINATE_SYSTEM_STAGE) failed!"));
+		// }
+	}
+	return XrResult::XR_SUCCESS;
+}
+
+XrResult FOXRVisionOSSession::XrCreateReferenceSpace(
+	const XrReferenceSpaceCreateInfo*			CreateInfo,
+	XrSpace*									Space)
+{
+	if (!CreateInfo || CreateInfo->type != XR_TYPE_REFERENCE_SPACE_CREATE_INFO)
+	{
+		return XrResult::XR_ERROR_VALIDATION_FAILURE;
+	}
+
+	TSharedPtr<FOXRVisionOSSpace, ESPMode::ThreadSafe> NewSpace = nullptr;
+	XrResult Ret = FOXRVisionOSSpace::CreateReferenceSpace(NewSpace, CreateInfo, this);
+	if (Ret == XrResult::XR_SUCCESS)
+	{
+		Spaces.Add(NewSpace);
+		*Space = (XrSpace)NewSpace.Get();
+	}
+
+	return Ret;
+}
+
+XrResult FOXRVisionOSSession::XrGetReferenceSpaceBoundsRect(
+	XrReferenceSpaceType                    ReferenceSpaceType,
+	XrExtent2Df*							Bounds)
+{
+	if (!Bounds)
+	{
+		return XrResult::XR_ERROR_VALIDATION_FAILURE;
+	}
+
+	if (ReferenceSpaceType != XR_REFERENCE_SPACE_TYPE_STAGE)
+	{
+		UE_LOG(LogOXRVisionOS, Log, TEXT("XrGetReferenceSpaceBoundsRect called with an reference space type that does not support it"));
+		Bounds->width = 0.0f;
+		Bounds->height = 0.0f;
+		return XR_SPACE_BOUNDS_UNAVAILABLE;
+	}
+
+	// According to docs the play area is a 1.5m sphere around the hmd.
+	// We will not return any bounds for now.
+	return XR_SPACE_BOUNDS_UNAVAILABLE;
+
+	// Bounds->width = ;
+	// Bounds->height = ;
+	// return XR_SUCCESS;
+}
+
+XrResult FOXRVisionOSSession::XrCreateActionSpace(
+	const XrActionSpaceCreateInfo*	CreateInfo,
+	XrSpace*						Space)
+{
+	if (!CreateInfo || CreateInfo->type != XR_TYPE_ACTION_SPACE_CREATE_INFO || !Space)
+	{
+		return XrResult::XR_ERROR_VALIDATION_FAILURE;
+	}
+
+	TSharedPtr<FOXRVisionOSSpace, ESPMode::ThreadSafe> NewSpace = nullptr;
+	XrResult Ret = FOXRVisionOSSpace::CreateActionSpace(NewSpace, CreateInfo, this);
+	if (Ret == XrResult::XR_SUCCESS)
+	{
+		Spaces.Add(NewSpace);
+		*Space = (XrSpace)NewSpace.Get();
+	}
+	return Ret;
+}
+
+XrResult FOXRVisionOSSession::DestroySpace(FOXRVisionOSSpace* Space)
+{
+	uint32 ArrayIndex = Spaces.IndexOfByPredicate([Space](const TSharedPtr<FOXRVisionOSSpace, ESPMode::ThreadSafe>& Data) { return (Data.Get() == Space); });
+
+	if (ArrayIndex == INDEX_NONE)
+	{
+		return  XrResult::XR_ERROR_HANDLE_INVALID;
+	}
+
+	Spaces.RemoveAtSwap(ArrayIndex);
+
+	return XrResult::XR_SUCCESS;
+}
+
+XrResult FOXRVisionOSSession::XrEnumerateSwapchainFormats(
+	uint32_t                                    FormatCapacityInput,
+	uint32_t*									FormatCountOutput,
+	int64_t*									Formats)
+{
+	if (FormatCountOutput == nullptr)
+	{
+		return XrResult::XR_ERROR_VALIDATION_FAILURE;
+	}
+	if (FormatCapacityInput != 0 && Formats == nullptr)
+	{
+		return XrResult::XR_ERROR_VALIDATION_FAILURE;
+	}
+
+	const int NumSupportedFormats = 4;
+
+	*FormatCountOutput = NumSupportedFormats;
+	if (FormatCapacityInput != 0)
+	{
+		if (FormatCapacityInput < *FormatCountOutput)
+		{
+			return XrResult::XR_ERROR_SIZE_INSUFFICIENT;
+		}
+
+		//TODO is this right?
+		// HDR formats 
+		Formats[0] = PF_B8G8R8A8;			//k8_8_8_8UNorm			// This is Unreal's default format
+		Formats[1] = PF_A2B10G10R10; 		//k10_10_10_2UNorm		
+		Formats[2] = PF_FloatRGB; 			//k11_11_10Float
+		Formats[3] = PF_FloatRGBA; 			//k16_16_16_16Float
+
+		//Depth formats - depth is not supported
+		//formats[1] = PF_DepthStencil;		//k32Float
+		//formats[2] = PF_D24;				//k32Float
+	}
+	return XrResult::XR_SUCCESS;
+}
+
+XrResult FOXRVisionOSSession::XrCreateSwapchain(
+	const XrSwapchainCreateInfo*				CreateInfo,
+	XrSwapchain*								Swapchain)
+{
+	if (CreateInfo == nullptr || CreateInfo->type != XR_TYPE_SWAPCHAIN_CREATE_INFO || Swapchain == nullptr)
+	{
+		return XrResult::XR_ERROR_VALIDATION_FAILURE;
+	}
+
+	TSharedPtr<FOXRVisionOSSwapchain, ESPMode::ThreadSafe> NewSwapchain = nullptr;
+	XrResult Ret = FOXRVisionOSSwapchain::Create(NewSwapchain, CreateInfo, OXRVisionOSSessionHelpers::OXRVisionOSBackbufferLength, this);
+	if (Ret == XrResult::XR_SUCCESS)
+	{
+		*Swapchain = (XrSwapchain)NewSwapchain.Get();
+		Swapchains.Add(NewSwapchain);
+	}
+	return Ret;
+}
+
+XrResult FOXRVisionOSSession::DestroySwapchain(
+	FOXRVisionOSSwapchain* Swapchain)
+{
+	uint32 ArrayIndex = Swapchains.IndexOfByPredicate([Swapchain](const TSharedPtr<FOXRVisionOSSwapchain, ESPMode::ThreadSafe>& Data) { return (Data.Get() == Swapchain); });
+
+	if (ArrayIndex == INDEX_NONE)
+	{
+		return XrResult::XR_ERROR_HANDLE_INVALID;
+	}
+
+	Swapchains.RemoveAtSwap(ArrayIndex);
+
+	return XrResult::XR_SUCCESS;
+}
+
+XrResult FOXRVisionOSSession::XrAttachSessionActionSets(
+	const XrSessionActionSetsAttachInfo* AttachInfo)
+{
+	if (AttachInfo == nullptr || AttachInfo->type != XR_TYPE_SESSION_ACTION_SETS_ATTACH_INFO)
+	{
+		return XrResult::XR_ERROR_VALIDATION_FAILURE;
+	}
+
+	if (ActionSetsAlreadyAttached)
+	{
+		return XrResult::XR_ERROR_ACTIONSETS_ALREADY_ATTACHED;
+	}
+
+	UE_LOG(LogOXRVisionOS, Verbose, TEXT("XrAttachSessionActionSets Attaching"));
+
+	check(AttachInfo->next == nullptr);
+	for (int i = 0; i < AttachInfo->countActionSets; i++)
+	{
+		XrActionSet ActionSet = AttachInfo->actionSets[i];
+		check(ActionSet != XR_NULL_HANDLE);
+		TSharedPtr<FOXRVisionOSActionSet, ESPMode::ThreadSafe> ActionSetPtr = Instance->GetActionSetSharedPtr(ActionSet);
+		ActionSetPtr->Attach(this);
+		ActionSets.Add(ActionSetPtr); 
+	}
+
+	// Choose the InteractionProfileBinding we will actually use.
+	bool bBoundAnInteractionProfile = false;
+	for (auto& InteractionProfilePath : Instance->GetPreferredInteractionProfileBindingList())
+	{
+		const FOXRVisionOSInstance::FInteractionProfileBinding* IPB = Instance->GetInteractionProfileBindings().Find(InteractionProfilePath);
+
+		if (IPB != nullptr)
+		{
+			UE_LOG(LogOXRVisionOS, Verbose, TEXT("XrAttachSessionActionSets found interaction profile %s and will bind inputs using it."), ANSI_TO_TCHAR(Instance->PathToString(IPB->InteractionProfilePath)));
+			bBoundAnInteractionProfile = true;
+
+			for (TSharedPtr<FOXRVisionOSActionSet, ESPMode::ThreadSafe>& ActionSet : ActionSets)
+			{
+				// Each action may be bound to one or more source
+				for (TSharedPtr<FOXRVisionOSAction, ESPMode::ThreadSafe> ActionSP : ActionSet->GetActions())
+				{
+					FOXRVisionOSAction* Action = ActionSP.Get();
+					for (const FOXRVisionOSInstance::FInteractionProfileBinding::FBinding& Binding : IPB->Bindings)
+					{
+						if (Binding.Action == (XrAction)Action)
+						{
+							// If this interaction profile is the current interaction profile of the top level path of this binding bind it.
+							XrPath TopLevelBindingPath = Instance->InputPathToTopLevelPath(Binding.Path);
+							XrPath TopLevelPathInteractionProfile = Instance->TopLevelPathToInteractionProfile(TopLevelBindingPath);
+							if (InteractionProfilePath == TopLevelPathInteractionProfile)
+							{
+								UE_LOG(LogOXRVisionOS, Verbose, TEXT("  Binding %s to %s"), ANSI_TO_TCHAR(Instance->PathToString(Binding.Path)), OXRVisionOS::OXRVisionOSControllerButtonToTCHAR(Binding.Button));
+								Action->BindSource(Binding.Path, Binding.Button);
+							}
+							else
+							{
+								UE_LOG(LogOXRVisionOS, Verbose, TEXT("  NOT binding %s to %s because this is not the best interaction profile"), ANSI_TO_TCHAR(Instance->PathToString(Binding.Path)), OXRVisionOS::OXRVisionOSControllerButtonToTCHAR(Binding.Button));
+							}
+						}
+					}
+				}
+			}
+		}
+	}
+	
+	if (!bBoundAnInteractionProfile)
+	{
+		UE_LOG(LogOXRVisionOS, Warning, TEXT("XrAttachSessionActionSets has been called, but no supported interaction profiles were suggested through xrSuggestInteractionProfileBindings!  None of the input actions will function."));
+	}
+	
+	ActionSetsAlreadyAttached = true;
+	return XrResult::XR_SUCCESS;
+}
+
+XrResult FOXRVisionOSSession::XrGetCurrentInteractionProfile(
+	XrPath TopLevelUserPath,
+	XrInteractionProfileState* InteractionProfile)
+{
+	if (ActionSetsAlreadyAttached == false)
+	{
+		return XrResult::XR_ERROR_ACTIONSET_NOT_ATTACHED;
+	}
+
+	InteractionProfile->interactionProfile = Instance->GetCurrentInteractionProfile(TopLevelUserPath);
+
+	return XrResult::XR_SUCCESS;
+}
+
+XrResult FOXRVisionOSSession::XrGetActionStateBoolean(
+	const XrActionStateGetInfo* GetInfo,
+	XrActionStateBoolean* State)
+{
+	const FOXRVisionOSAction* Action = (const FOXRVisionOSAction*)GetInfo->action;
+	if (Action)
+	{
+		return Action->GetActionStateBoolean(GetInfo->subactionPath, *State);
+	}
+	else
+	{
+		return XR_ERROR_HANDLE_INVALID;
+	}
+}
+
+XrResult FOXRVisionOSSession::XrGetActionStateFloat(
+	const XrActionStateGetInfo* GetInfo,
+	XrActionStateFloat* State)
+{
+	const FOXRVisionOSAction* Action = (const FOXRVisionOSAction*)GetInfo->action;
+	if (Action)
+	{
+		return Action->GetActionStateFloat(GetInfo->subactionPath, *State);
+	}
+	else
+	{
+		return XR_ERROR_HANDLE_INVALID;
+	}
+}
+
+XrResult FOXRVisionOSSession::XrGetActionStateVector2f(
+	const XrActionStateGetInfo* GetInfo,
+	XrActionStateVector2f* State)
+{
+	const FOXRVisionOSAction* Action = (const FOXRVisionOSAction*)GetInfo->action;
+	if (Action)
+	{
+		return Action->GetActionStateVector2f(GetInfo->subactionPath, *State);
+	}
+	else
+	{
+		return XR_ERROR_HANDLE_INVALID;
+	}
+}
+
+XrResult FOXRVisionOSSession::XrGetActionStatePose(
+	const XrActionStateGetInfo* GetInfo,
+	XrActionStatePose* State)
+{
+	const FOXRVisionOSAction* Action = (const FOXRVisionOSAction*)GetInfo->action;
+	if (Action)
+	{
+		return Action->GetActionStatePose(GetInfo->subactionPath, *State);
+	}
+	else
+	{
+		return XR_ERROR_HANDLE_INVALID;
+	}
+}
+
+XrResult FOXRVisionOSSession::XrSyncActions(
+	const XrActionsSyncInfo* SyncInfo)
+{
+	check(SyncInfo);
+
+	//TODO XR_SESSION_LOSS_PENDING
+
+	if (SessionState != XR_SESSION_STATE_FOCUSED)
+	{
+		return XrResult::XR_SESSION_NOT_FOCUSED;
+	}
+
+	// Collect the hardware data
+	//Controllers->SyncActions();
+	const FPipelinedFrameState& PipelinedFrameState = GetPipelinedFrameStateForThread();
+	Tracker->SyncActions(PipelinedFrameState.FrameCounter, PipelinedFrameState.PredictedDisplayTime);
+
+	// Clear active from all action sets
+	for (TSharedPtr<FOXRVisionOSActionSet, ESPMode::ThreadSafe>& ActionSet : ActionSets)
+	{
+		ActionSet->ClearActive();
+	}
+	
+	//TArray<FOXRVisionOSActionSet*> ActiveActionSetsByPriority; // Note ActionSetPriority is currently unsupported on OXRVisionOS.
+	for (int32 i = 0; i < SyncInfo->countActiveActionSets; i++)
+	{
+		XrActiveActionSet ActiveXRActionSet = SyncInfo->activeActionSets[i];
+		FOXRVisionOSActionSet& ActiveActionSet = *((FOXRVisionOSActionSet*)ActiveXRActionSet.actionSet);
+		ActiveActionSet.SetActive(ActiveXRActionSet.subactionPath);
+		//ActiveActionSetsByPriority.Add(&ActiveActionSet);
+	}
+	// Sort highest priority to the front.
+	//ActiveActionSetsByPriority.Sort([](const FOXRVisionOSActionSet& A, const FOXRVisionOSActionSet& B) { return A.GetPriority() < B.GetPriority(); });
+
+//UE_LOG(LogOXRVisionOS, Log, TEXT("XrSyncActions PipelinedFrameState.PredictedDisplayTime %lld"), PipelinedFrameState.PredictedDisplayTime);
+
+	// Go through the action sets updating state per hardware.
+//	for (TSharedPtr<FOXRVisionOSActionSet, ESPMode::ThreadSafe>& ActionSet : ActionSets)
+//	{
+//		ActionSet->SyncActions(Controllers, Tracker);
+//	}
+
+	return XR_SUCCESS;
+}
+
+XrResult FOXRVisionOSSession::XrEnumerateBoundSourcesForAction(
+	const XrBoundSourcesForActionEnumerateInfo*		EnumerateInfo,
+	uint32_t										SourceCapacityInput,
+	uint32_t*										SourceCountOutput,
+	XrPath*											OutSources)
+{
+	if (SourceCountOutput == nullptr)
+	{
+		return XrResult::XR_ERROR_VALIDATION_FAILURE;
+	}
+	if (SourceCapacityInput != 0 && OutSources == nullptr)
+	{
+		return XrResult::XR_ERROR_VALIDATION_FAILURE;
+	}
+
+	if (EnumerateInfo == nullptr || EnumerateInfo->action == XR_NULL_HANDLE)
+	{
+		return XrResult::XR_ERROR_VALIDATION_FAILURE;
+	}
+
+	FOXRVisionOSAction* Action = (FOXRVisionOSAction*)EnumerateInfo->action;
+	const TArray<XrPath>& Sources = Action->GetBoundSources();
+
+	const int Num = Sources.Num();
+	*SourceCountOutput = Num;
+	if (SourceCapacityInput != 0)
+	{
+		if (SourceCapacityInput < *SourceCountOutput)
+		{
+			return XrResult::XR_ERROR_SIZE_INSUFFICIENT;
+		}
+
+		for (int32 i = 0; i < Num; ++i)
+		{
+			OutSources[i] = Sources[i];
+		}
+	}
+	return XrResult::XR_SUCCESS;
+}
+
+XrResult FOXRVisionOSSession::XrApplyHapticFeedback(
+	const XrHapticActionInfo* HapticActionInfo,
+	const XrHapticBaseHeader* HapticFeedback)
+{
+	if (SessionState != XR_SESSION_STATE_FOCUSED)
+	{
+		return XrResult::XR_SESSION_NOT_FOCUSED;
+	}
+
+	UE_LOG(LogOXRVisionOS, Warning, TEXT("FOXRVisionOSSession XrApplyHapticFeedback returning XR_ERROR_PATH_UNSUPPORTED for path because there are no supported haptic paths on this platform."));
+	return XrResult::XR_ERROR_PATH_UNSUPPORTED;
+}
+
+XrResult FOXRVisionOSSession::XrStopHapticFeedback(
+	const XrHapticActionInfo* HapticActionInfo)
+{
+	if (SessionState != XR_SESSION_STATE_FOCUSED)
+	{
+		return XrResult::XR_SESSION_NOT_FOCUSED;
+	}
+
+	UE_LOG(LogOXRVisionOS, Warning, TEXT("FOXRVisionOSSession XrApplyHapticFeedback returning XR_ERROR_PATH_UNSUPPORTED for path because there are no supported haptic paths on this platform."));
+	return XrResult::XR_ERROR_PATH_UNSUPPORTED;
+}
+
+XrResult FOXRVisionOSSession::XrBeginSession(
+	const XrSessionBeginInfo* BeginInfo)
+{
+	if (BeginInfo == nullptr || BeginInfo->type != XR_TYPE_SESSION_BEGIN_INFO)
+	{
+		return XrResult::XR_ERROR_VALIDATION_FAILURE;
+	}
+
+	if (SessionState == XrSessionState::XR_SESSION_STATE_UNKNOWN  || SessionState == XrSessionState::XR_SESSION_STATE_IDLE)
+	{
+		return XrResult::XR_ERROR_SESSION_NOT_READY;
+	}
+
+	if (IsSessionRunning())
+	{
+		return XrResult::XR_ERROR_SESSION_RUNNING;
+	}
+
+	if (BeginInfo->primaryViewConfigurationType != XrViewConfigurationType::XR_VIEW_CONFIGURATION_TYPE_PRIMARY_STEREO)
+	{
+		return XrResult::XR_ERROR_VIEW_CONFIGURATION_TYPE_UNSUPPORTED;
+	}
+    
+    VOSLayerRenderer = OXRVisionOS::GetSwiftLayerRenderer();
+
+	// Create World Tracking data provider
+	ar_world_tracking_configuration_t ArKitWorldTrackingConfiguration = ar_world_tracking_configuration_create();
+	ARKitWorldTrackingProvider = ar_world_tracking_provider_create(ArKitWorldTrackingConfiguration);
+	// Create Data providers
+	ar_data_providers_t ARKitDataProviders = ar_data_providers_create();
+	// Add
+	ar_data_providers_add_data_provider(ARKitDataProviders, ARKitWorldTrackingProvider);
+	
+	// Start the arkit session
+	ar_session_run(ARKitSession, ARKitDataProviders);
+
+	// If we failed to get access to the hmd we could fail here like this.
+	// if (???)
+	// {
+	// 	UE_LOG(LogOXRVisionOS, Warning, TEXT("Open failed code 0x%x.  Session going to XR_SESSION_STATE_LOSS_PENDING."), HMDHandle);
+	// 	SetSessionState(XrSessionState::XR_SESSION_STATE_LOSS_PENDING);
+	// 	return XrResult::XR_SESSION_LOSS_PENDING;
+	// }
+
+	if (Tracker->RegisterDevice(TrackerDeviceType::DEVICE_VR_HMD, HMDHandle) == false)
+	{
+		UE_LOG(LogOXRVisionOS, Error, TEXT("Tracker->RegisterDevice failed.  Bug."));
+		return XrResult::XR_ERROR_RUNTIME_FAILURE;
+	}
+
+    // Fill the game threads view info buffer with temporary data, some of which was collected in FSwiftAppBootstrap::KickoffWithCompositingLayer
+	// XrLocateViews will overwrite this with real data fetched from the VisionOS API, but we will use this temp fov data in the first two game thread frames.
+    LocateViewInfoBuffer[0].NumViews = OXRVisionOS::GetSwiftNumViewports();
+    LocateViewInfoBuffer[0].ViewTransforms[0] = matrix_identity_float4x4;
+    LocateViewInfoBuffer[0].ViewTransforms[1] = matrix_identity_float4x4;
+    simd_float4x4 ViewTransforms[2];
+    const float HalfFov = 2.19686294f / 2.f;
+    LocateViewInfoBuffer[0].HmdFovs[0].angleLeft = -HalfFov - 0.1f;
+    LocateViewInfoBuffer[0].HmdFovs[0].angleRight = HalfFov - 0.1f;
+    LocateViewInfoBuffer[0].HmdFovs[0].angleUp = HalfFov;
+    LocateViewInfoBuffer[0].HmdFovs[0].angleDown = -HalfFov;
+    LocateViewInfoBuffer[0].HmdFovs[1].angleLeft = -HalfFov + 0.1f;
+    LocateViewInfoBuffer[0].HmdFovs[1].angleRight = HalfFov + 0.1f;
+    LocateViewInfoBuffer[0].HmdFovs[1].angleUp = HalfFov;
+    LocateViewInfoBuffer[0].HmdFovs[1].angleDown = HalfFov;
+    LocateViewInfoBuffer[0].HmdFovs[1] = LocateViewInfoBuffer[0].HmdFovs[0];
+    LocateViewInfoBuffer[1] = LocateViewInfoBuffer[0];
+    LocateViewInfoBuffer[2] = LocateViewInfoBuffer[0];
+
+	bIsRunning = true;
+	PipelinedFrameStateGame.bSynchronizing = true;
+	PipelinedFrameStateRendering.bSynchronizing = true;
+	PipelinedFrameStateRHI.bSynchronizing = true;
+
+	//Controllers->OnBeginSession(Tracker.Get());
+
+	return XrResult::XR_SUCCESS;
+}
+
+XrResult FOXRVisionOSSession::XrEndSession()
+{
+	if (!IsSessionRunning())
+	{
+		return XrResult::XR_ERROR_SESSION_NOT_RUNNING;
+	}
+
+	if (SessionState != XrSessionState::XR_SESSION_STATE_STOPPING)
+	{
+		return XrResult::XR_ERROR_SESSION_NOT_STOPPING;
+	}
+
+	EndSessionInternal();
+
+	SetSessionState(XrSessionState::XR_SESSION_STATE_IDLE);
+	// Start polling to go back to XR_SESSION_STATE_READY as soon as the hmd is available
+	OnWorldTickStartDelegateHandle = FWorldDelegates::OnWorldTickStart.AddRaw(this, &FOXRVisionOSSession::OnWorldTickStart);
+
+	return XrResult::XR_SUCCESS;
+}
+
+void FOXRVisionOSSession::EndSessionInternal()
+{
+	UE_LOG(LogOXRVisionOS, Warning, TEXT("FOXRVisionOSSession::EndSessionInternal()"));
+
+	//Controllers->OnEndSession(Tracker.Get());
+
+	Tracker->UnregisterDevice(HMDHandle);
+	//Close(HMDHandle);
+	HMDHandle = INVALID_DEVICE_HANDLE;
+
+	PipelinedFrameStateGame.CommandListContext = PipelinedFrameStateRendering.CommandListContext = PipelinedFrameStateRHI.CommandListContext = nullptr;
+
+	if (bExitRequested)
+	{
+		SetSessionState(XrSessionState::XR_SESSION_STATE_EXITING);
+	}
+
+	bExitRequested = false;
+	bIsRunning = false;
+}
+
+XrResult FOXRVisionOSSession::XrRequestExitSession()
+{
+	if (!IsSessionRunning())
+	{
+		return XrResult::XR_ERROR_SESSION_NOT_RUNNING;
+	}
+
+	SetSessionState(XrSessionState::XR_SESSION_STATE_STOPPING);
+	bExitRequested = true;
+	return XrResult::XR_SUCCESS;
+}
+
+// FOpenXRHMD::OnBeginSimulation_GameThread()
+XrResult FOXRVisionOSSession::XrWaitFrame(
+	const XrFrameWaitInfo* frameWaitInfo,
+	XrFrameState* frameState)
+{
+	//UE_LOG(LogOXRVisionOS, Log, TEXT("FOXRVisionOSSession XrWaitFrame called"));
+
+	SCOPED_NAMED_EVENT_TEXT("FOXRVisionOSSession::XrWaitFrame", FColor::Turquoise);
+
+	// frameWaitInfo can be null
+	if (frameWaitInfo != nullptr && frameWaitInfo->type != XR_TYPE_FRAME_WAIT_INFO)
+	{
+		return XrResult::XR_ERROR_VALIDATION_FAILURE;
+	}
+
+	if (frameState == nullptr || frameState->type != XR_TYPE_FRAME_STATE)
+	{
+		return XrResult::XR_ERROR_VALIDATION_FAILURE;
+	}
+
+	if (!IsSessionRunning())
+	{
+		UE_LOG(LogOXRVisionOS, Log, TEXT("FOXRVisionOSSession XrWaitFrame XR_ERROR_SESSION_NOT_RUNNING"));
+		return XrResult::XR_ERROR_SESSION_NOT_RUNNING;
+	}
+
+	if (GetSessionState() == XrSessionState::XR_SESSION_STATE_READY)
+	{
+		SetSessionState(XrSessionState::XR_SESSION_STATE_SYNCHRONIZED);
+		SetSessionState(XrSessionState::XR_SESSION_STATE_VISIBLE);
+		SetSessionState(XrSessionState::XR_SESSION_STATE_FOCUSED);
+
+		SessionFrameCounter += OXRVisionOSSessionHelpers::OXRVisionOSBackbufferLength;  // Jump ahead in case the current frame counter was partially used while the previous session was being ended.  
+		PipelinedFrameStateGame.bSynchronizing = false;
+	}
+
+	// Block until the previous frame's xrBeginFrame has happened.
+	{
+		SCOPED_NAMED_EVENT_TEXT("FOXRVisionOSSession::XrWaitFrame XrWaitFrameEvent->Wait()", FColor::Turquoise);
+		XrWaitFrameEvent->Wait();
+	}
+
+	PipelinedFrameStateGame.FrameCounter = SessionFrameCounter++;
+
+	//UE_LOG(LogOXRVisionOS, Log, TEXT("FOXRVisionOSSession XrWaitFrame proceeds Game fliparg: %d BBIndex: %d"), PipelinedFrameStateGame.FrameCounter, PipelinedFrameStateGame.BackBufferIndex);
+
+	// {
+	// 	const int32 ret = ReprojectionBeginFrame(PipelinedFrameStateGame.FrameCounter);
+	// 	if (ret != RESULT_OK)
+	// 	{
+	// 		UE_LOG(LogOXRVisionOS, Warning, TEXT("ReprojectionBeginFrame(%d) failed: %d (%u, 0x%08X)"), PipelinedFrameStateGame.FrameCounter, ret, ret, ret);
+	// 	}
+	// }
+    
+    PipelinedFrameStateGame.SwiftFrame = cp_layer_renderer_query_next_frame(VOSLayerRenderer);
+    UE_LOG(LogTemp, Log, TEXT("SwiftLayerFrame(0x%x) from cp_layer_renderer_query_next_frame in FOXRVisionOSSession::XrWaitFrame"), PipelinedFrameStateGame.SwiftFrame);
+    
+    // Fetch the predicted timing information.
+    PipelinedFrameStateGame.SwiftFrameTiming = cp_frame_predict_timing(PipelinedFrameStateGame.SwiftFrame);
+    if ( PipelinedFrameStateGame.SwiftFrameTiming == nullptr)
+    {
+        //TODO ??? failed frame?
+        // This means the layer is not in the correct state  Perhaps we are either running too early or we need to end the session.
+        // we might want to return XR_ERROR_SESSION_NOT_RUNNING here, but only if we actually end the session first.
+        assert(false);
+        //return;
+    }
+
+    // Mark the beginnng of our game thread update work.  All the stuff before rendering.
+    UE_LOG(LogTemp, Log, TEXT("SwiftLayerFrame(0x%x) from cp_frame_start_update in FOXRVisionOSSession::XrWaitFrame"), PipelinedFrameStateGame.SwiftFrame);
+    cp_frame_start_update(PipelinedFrameStateGame.SwiftFrame);
+
+    const cp_time_t PresentationTime = cp_frame_timing_get_presentation_time(PipelinedFrameStateGame.SwiftFrameTiming);
+    const CFTimeInterval TargetRenderTime = cp_time_to_cf_time_interval(PresentationTime); // CRTimeInterval is double time in seconds.
+    
+	XrTime DisplayPeriod = 6 * 1000000; //TODO other frame sequences?  This is 90fps
+	frameState->predictedDisplayPeriod = DisplayPeriod;
+	frameState->predictedDisplayTime = OXRVisionOS::APITimeToXrTime(TargetRenderTime);
+	PipelinedFrameStateGame.PredictedDisplayTime = frameState->predictedDisplayTime;
+    
+    PipelinedFrameStateGame.DeviceAnchor = ar_device_anchor_create();
+    check(ARKitWorldTrackingProvider)
+    auto anchor_status = ar_world_tracking_provider_query_device_anchor_at_timestamp(ARKitWorldTrackingProvider, TargetRenderTime, PipelinedFrameStateGame.DeviceAnchor);
+    if (anchor_status == ar_device_anchor_query_status_success)
+    {
+        // Get the HMD transform and cache it
+        PipelinedFrameStateGame.HeadTransform = ar_anchor_get_origin_from_anchor_transform(PipelinedFrameStateGame.DeviceAnchor);
+    }
+    else
+    {
+        // Cache identity as the HMD transform as a fallback
+        PipelinedFrameStateGame.HeadTransform = matrix_identity_float4x4;
+    }
+    
+    LocateViewInfoIndexAdvance();
+
+	//TODO XrWaitFrame shouldRender based on visible & focused state
+	frameState->shouldRender = XR_TRUE;
+
+	// We have to pipeline state updates in order to prevent task overlap (esp during map loading where we kick extra renders)
+	ENQUEUE_RENDER_COMMAND(UpdatePipelinedFrameState)([this, NewRenderingFrameState = PipelinedFrameStateGame](FRHICommandListImmediate& RHICmdList)
+	{
+		check(IsInRenderingThread());
+		PipelinedFrameStateRendering = NewRenderingFrameState;
+		RHICmdList.EnqueueLambda([this, NewRHIFrameState = PipelinedFrameStateRendering](FRHICommandListImmediate&)
+		{
+			PipelinedFrameStateRHI = NewRHIFrameState;
+		});
+	});
+
+	// Update the STAGE space
+	//Tracker->UpdateReferenceSpaces(Instance, this, HMDHandle, PipelinedFrameStateGame.PredictedDisplayTime, PipelinedFrameStateGame.FrameCounter);
+
+	return XrResult::XR_SUCCESS;
+}
+
+void FOXRVisionOSSession::OnBeginRendering_GameThread()
+{
+    // Mark the end of our previous game frame
+    if (PipelinedFrameStateGame.SwiftFrame)
+    {
+        UE_LOG(LogTemp, Log, TEXT("SwiftLayerFrame(0x%x) from cp_frame_end_update in FOXRVisionOSSession::OnBeginRendering_GameThread"), PipelinedFrameStateGame.SwiftFrame);
+        cp_frame_end_update(PipelinedFrameStateGame.SwiftFrame);
+        PipelinedFrameStateGame.SwiftFrame = nullptr;
+    }
+}
+
+// FOpenXRHMD::OnBeginRendering_RenderThread
+void FOXRVisionOSSession::OnBeginRendering_RenderThread()
+{
+    const MetalRHIVisionOS::BeginRenderingImmersiveParams Params{ PipelinedFrameStateRendering.SwiftFrame };
+    MetalRHIVisionOS::BeginRenderingImmersive(Params);
+    WaitUntil();
+    StartSubmission();
+}
+
+// FOpenXRHMD::OnBeginRendering_RHIThread
+XrResult FOXRVisionOSSession::XrBeginFrame(
+	const XrFrameBeginInfo* FrameBeginInfo)
+{
+	check(!IsRHIThreadRunning() || IsInRHIThread());
+
+	//UE_LOG(LogOXRVisionOS, Log, TEXT("FOXRVisionOSSession XrBeginFrame fliparg: %d BBIndex: %d"), PipelinedFrameStateRendering.FrameCounter, PipelinedFrameStateRendering.BackBufferIndex); // Note frame state not copied yet
+
+	SCOPED_NAMED_EVENT_TEXT("FOXRVisionOSSession::XrBeginFrame", FColor::Turquoise);
+
+	// frameBeginInfo can be null
+	if (FrameBeginInfo != nullptr && FrameBeginInfo->type != XR_TYPE_FRAME_BEGIN_INFO)
+	{
+		XrWaitFrameEvent->Trigger();
+		return XrResult::XR_ERROR_VALIDATION_FAILURE;
+	}
+
+	if (!IsSessionRunning())
+	{
+		XrWaitFrameEvent->Trigger();
+		return XrResult::XR_ERROR_SESSION_NOT_RUNNING;
+	}
+
+	if (PipelinedFrameStateRHI.FrameCounter == CachedBeginFlipFrameCounter)
+	{
+		// We don't trigger in this case because there wasn't a call to xrWaitFrame that corresponds
+		// xrWaitFrame should have incremented the counter
+		UE_LOG(LogOXRVisionOS, Log, TEXT("XrWaitFrameEvent not triggered because call order invalid  RHIFrameCounter %i != RenderFrameCounter %i"), PipelinedFrameStateRHI.FrameCounter, CachedBeginFlipFrameCounter);
+		return XrResult::XR_ERROR_CALL_ORDER_INVALID;
+	}
+    CachedBeginFlipFrameCounter = PipelinedFrameStateRHI.FrameCounter;
+       
+	PipelinedFrameStateRHI.CommandListContext = RHIGetDefaultContext();
+	
+	if (PipelinedFrameStateRHI.bSynchronizing)
+	{
+		XrWaitFrameEvent->Trigger();
+		//UE_LOG(LogOXRVisionOS, Log, TEXT("FOXRVisionOSSession XrBeginFrame bSynchronizing=true RHI fliparg: %d BBIndex: %d"), PipelinedFrameStateRHI.FrameCounter, PipelinedFrameStateRHI.BackBufferIndex);
+		return XrResult::XR_SUCCESS;
+	}
+
+	XrWaitFrameEvent->Trigger();
+
+	const XrResult BeginFrameResult = (NextBackBufferIndex != PreviousBackBufferIndexInBegin) ? XrResult::XR_SUCCESS : XrResult::XR_FRAME_DISCARDED;
+
+	//if (PipelinedFrameStateRHI.CommandListContext)
+    {
+        // TODO
+        // I think there's gonna be some trickiness when we have accurate VISIBLE + FOCUSED session states.
+        // If we return XrFrameState.shouldRender == false (headset off), we still could deploy WaitUntilSafeForRendering
+        // either here or in xrWaitSwapchainImage. In that case, we will hang because the flip won't deploy.
+        // We might want to reconsider using WaitUntilSafeForRendering. If we use xrWaitFrame or xrWaitSwapchainImage
+        // to wait for a CPU-signal for swapchain safety, then...we don't really need to check on the GPU again.
+        // But the docs infer that ReprojectionBeginFrame and waitUntilSafeForRendering are tied :/
+    }
+    
+    {
+        //TODO idk, doing this twice once for render once for rhi...
+        FPipelinedFrameState& FrameState = PipelinedFrameStateRHI;
+        
+        UE_LOG(LogTemp, Log, TEXT("SwiftLayerFrame(0x%x) cp_frame_query_drawable in FOXRVisionOSSession::XrBeginFrame"), FrameState.SwiftFrame);
+        FrameState.SwiftDrawable = cp_frame_query_drawable(FrameState.SwiftFrame);
+        check(FrameState.SwiftDrawable);
+        UE_LOG(LogTemp, Log, TEXT("SwiftLayerFrame(0x%x) cp_frame_query_drawable got SwiftDrawable(0x%x) in FOXRVisionOSSession::StartSubmission()"), FrameState.SwiftFrame, FrameState.SwiftDrawable);
+
+        FrameState.SwiftFinalFrameTiming = cp_drawable_get_frame_timing(FrameState.SwiftDrawable);
+        
+        CFTimeInterval TimeInterval = cp_time_to_cf_time_interval(cp_frame_timing_get_presentation_time(FrameState.SwiftFinalFrameTiming));
+        check(ARKitWorldTrackingProvider)
+        auto anchor_status = ar_world_tracking_provider_query_device_anchor_at_timestamp(ARKitWorldTrackingProvider, TimeInterval, FrameState.DeviceAnchor);
+        if (anchor_status == ar_device_anchor_query_status_success)
+        {
+            // Get the HMD transform and cache it
+            FrameState.HeadTransform = ar_anchor_get_origin_from_anchor_transform(FrameState.DeviceAnchor);
+        }
+        else
+        {
+            // Leave FrameState.HeadTransform untouched, same as game thread. It may simply be identity.  But whatever it is we will go ahead and render with it.
+        }
+    }
+    
+    //MetalRHIOXRVisionOS::WaitUntilSafeForRendering(PipelinedFrameStateRHI.CommandListContext, NextBackBufferIndex);
+    PreviousBackBufferIndexInBegin = NextBackBufferIndex;
+
+	return BeginFrameResult;
+}
+
+void FOXRVisionOSSession::WaitUntil()
+{
+    FPipelinedFrameState& FrameState = GetPipelinedFrameStateForThread();
+    
+    UE_LOG(LogTemp, Log, TEXT("SwiftLayerFrame(0x%x) cp_time_wait_until starting in FOXRVisionOSSession::WaitUntil"), FrameState.SwiftFrame);
+    cp_time_wait_until(cp_frame_timing_get_optimal_input_time(FrameState.SwiftFrameTiming));
+    UE_LOG(LogTemp, Log, TEXT("SwiftLayerFrame(0x%x) cp_time_wait_until completed in FOXRVisionOSSession::WaitUntil"), FrameState.SwiftFrame);
+}
+
+void FOXRVisionOSSession::StartSubmission()
+{
+    FPipelinedFrameState& FrameState = GetPipelinedFrameStateForThread();
+    
+    // tell the compositor service we are going to submit soon and need the drawable
+    UE_LOG(LogTemp, Log, TEXT("SwiftLayerFrame(0x%x) cp_frame_start_submission in FOXRVisionOSSession::XrBeginFrame"), FrameState.SwiftFrame);
+    cp_frame_start_submission(FrameState.SwiftFrame);
+
+    UE_LOG(LogTemp, Log, TEXT("SwiftLayerFrame(0x%x) cp_frame_query_drawable in FOXRVisionOSSession::XrBeginFrame"), FrameState.SwiftFrame);
+    FrameState.SwiftDrawable = cp_frame_query_drawable(FrameState.SwiftFrame);
+    check(FrameState.SwiftDrawable);
+    UE_LOG(LogTemp, Log, TEXT("SwiftLayerFrame(0x%x) cp_frame_query_drawable got SwiftDrawable(0x%x) in FOXRVisionOSSession::StartSubmission()"), FrameState.SwiftFrame, FrameState.SwiftDrawable);
+
+    FrameState.SwiftFinalFrameTiming = cp_drawable_get_frame_timing(FrameState.SwiftDrawable);
+    
+    CFTimeInterval TimeInterval = cp_time_to_cf_time_interval(cp_frame_timing_get_presentation_time(FrameState.SwiftFinalFrameTiming));
+    check(ARKitWorldTrackingProvider)
+    auto anchor_status = ar_world_tracking_provider_query_device_anchor_at_timestamp(ARKitWorldTrackingProvider, TimeInterval, FrameState.DeviceAnchor);
+    if (anchor_status == ar_device_anchor_query_status_success)
+    {
+        cp_drawable_set_device_anchor(FrameState.SwiftDrawable, FrameState.DeviceAnchor);
+        // Get the HMD transform and cache it
+        FrameState.HeadTransform = ar_anchor_get_origin_from_anchor_transform(FrameState.DeviceAnchor);
+    }
+    else
+    {
+        //TODO: this is also how you disable reprojection.  Hook that up to a cvar.
+        cp_drawable_set_device_anchor(FrameState.SwiftDrawable, nullptr);
+        // Leave FrameState.HeadTransform untouched, same as game thread. It may simply be identity.  But whatever it is we will go ahead and render with it.
+    }
+    
+    //cp_drawable_set_ar_pose(FrameState.SwiftDrawable, )  // I think this is may be a deprecated way of specifying the frame's transform.  Replaced with cp_drawable_set_device_anchor
+}
+
+// FOpenXRHMD::OnFinishRendering_RHIThread()
+XrResult FOXRVisionOSSession::XrEndFrame(
+	const XrFrameEndInfo* InFrameEndInfo)
+{
+	//UE_LOG(LogOXRVisionOS, Log, TEXT("FOXRVisionOSSession XrEndFrame RHI fliparg: %d BBIndex: %d"), PipelinedFrameStateRHI.FrameCounter, PipelinedFrameStateRHI.BackBufferIndex);
+
+	SCOPED_NAMED_EVENT_TEXT("FOXRVisionOSSession::XrEndFrame", FColor::Turquoise);
+
+	if (InFrameEndInfo == nullptr || InFrameEndInfo->type != XR_TYPE_FRAME_END_INFO)
+	{
+		return XrResult::XR_ERROR_VALIDATION_FAILURE;
+	}
+
+	if (!IsSessionRunning())
+	{
+		return XrResult::XR_ERROR_SESSION_NOT_RUNNING;
+	}
+
+	if (PipelinedFrameStateRHI.bSynchronizing)
+	{
+		//UE_LOG(LogOXRVisionOS, Log, TEXT("FOXRVisionOSSession XrEndFrame bSynchronizing=true fliparg=%i buffer=%i"), PipelinedFrameStateRHI.FrameCounter, PipelinedFrameStateRHI.BackBufferIndex);
+		return XrResult::XR_SUCCESS;
+	}
+
+	const XrFrameEndInfo& FrameEndInfo = *InFrameEndInfo;
+	check(FrameEndInfo.environmentBlendMode == XR_ENVIRONMENT_BLEND_MODE_OPAQUE); // only opaque is supported, see XrEnumerateEnvironmentBlendModes
+	check(FrameEndInfo.layerCount <= OXRVisionOSSessionHelpers::MaxOXRVisionOSLayers);
+    
+    {
+        //TODO??? supporting only one layer now.  This means we do not have a non-reprojected layer.  Maybe we could have one as a separate swiftui?
+        check(FrameEndInfo.layerCount == 1);
+        check(FrameEndInfo.layers != nullptr);
+        const XrCompositionLayerBaseHeader& LayerBaseHeader = *FrameEndInfo.layers[0];
+        check(LayerBaseHeader.type == XR_TYPE_COMPOSITION_LAYER_PROJECTION);
+        const XrCompositionLayerProjection& LayerProjection = *(reinterpret_cast<const XrCompositionLayerProjection*>(FrameEndInfo.layers[0]));
+        check(LayerProjection.viewCount == 2); // Must be the number of views returned by xrLocateViews
+        const XrCompositionLayerProjectionView& View0 = LayerProjection.views[0];
+        const XrCompositionLayerProjectionView& View1 = LayerProjection.views[1];
+
+        check(View0.subImage.swapchain == View1.subImage.swapchain); //TODO: should we do it this way or should we have two separate swapchains?  If we did we would need to handle both below.
+        
+        const FOXRVisionOSSwapchain& SwapchainData0 = *(reinterpret_cast<FOXRVisionOSSwapchain*>(View0.subImage.swapchain));
+        const FOXRVisionOSSwapchain::FSwapchainImage& ReleasedImage = SwapchainData0.GetLastReleasedImage();
+        
+        //TODO feed in all necessary cp_frame_end_submission data
+        const MetalRHIVisionOS::PresentImmersiveParams Params{ReleasedImage.Image, PipelinedFrameStateRHI.SwiftDrawable};
+        MetalRHIVisionOS::PresentImmersive(Params);
+    }
+
+	 	NextBackBufferIndex = (NextBackBufferIndex + 1) % OXRVisionOSSessionHelpers::OXRVisionOSBackbufferLength;
+	
+	// else
+	// {
+	// 	UE_LOG(LogOXRVisionOS, Log, TEXT("XrEndFrame no layers, discarding, FrameCounter: %d, BufferIndex: %d"), PipelinedFrameStateRHI.FrameCounter, NextBackBufferIndex);
+	// }
+
+	return XrResult::XR_SUCCESS;
+}
+
+XrResult FOXRVisionOSSession::XrLocateViews(
+	const XrViewLocateInfo*		ViewLocateInfo,
+	XrViewState*				ViewState,
+	uint32_t                    ViewCapacityInput,
+	uint32_t*					ViewCountOutput,
+	XrView*						Views)
+{
+    if (ViewLocateInfo == nullptr || ViewLocateInfo->type != XR_TYPE_VIEW_LOCATE_INFO)
+    {
+        return XrResult::XR_ERROR_VALIDATION_FAILURE;
+    }
+    
+    if (ViewState == nullptr)
+    {
+        return XrResult::XR_ERROR_VALIDATION_FAILURE;
+    }
+    
+    if (ViewCapacityInput != 0 && Views == nullptr)
+    {
+        return XrResult::XR_ERROR_VALIDATION_FAILURE;
+    }
+    
+    if (ViewLocateInfo->viewConfigurationType != XR_VIEW_CONFIGURATION_TYPE_PRIMARY_STEREO)
+    {
+        return XrResult::XR_ERROR_VIEW_CONFIGURATION_TYPE_UNSUPPORTED;
+    }
+    
+    FPipelinedFrameState& FrameState = GetPipelinedFrameStateForThread();
+    
+    FOXRVisionOSSpace& BaseSpace = *(FOXRVisionOSSpace*)ViewLocateInfo->space;
+    const int32 FrameCounter = FrameState.FrameCounter;
+    FTransform BaseSpaceInverseTransform;
+    XrSpaceLocationFlags BaseSpaceFlags;
+    FVector OutLinearVelocity; FVector OutAngularVelocity; XrSpaceVelocityFlags OutVelocityFlags; // These are unused here.
+    BaseSpace.GetInverseTransform(ViewLocateInfo->displayTime, BaseSpaceInverseTransform, BaseSpaceFlags, OutLinearVelocity, OutAngularVelocity, OutVelocityFlags);
+    
+    uint32_t NumViews = 0;
+    if (IsInGameThread())
+    {
+        // We will use the NumViews from a previous render frame.
+        const FLocateViewInfo& ViewInfoCache = GetLocateViewInfo_GameThread();
+        NumViews = ViewInfoCache.NumViews;
+    }
+    else
+    {
+        check(IsInRenderingThread());
+        check(FrameState.SwiftDrawable);
+        NumViews = cp_drawable_get_view_count(FrameState.SwiftDrawable);
+        FLocateViewInfo& ViewInfoCache = GetLocateViewInfo_RenderThread();
+        ViewInfoCache.NumViews = NumViews;
+    }
+
+    //TODO probably all of this should happen only once per frame in xrBeginFrame and once in xrWaitFrame.
+    
+	*ViewCountOutput = NumViews;
+	if (ViewCapacityInput != 0)
+    {
+        if (ViewCapacityInput < *ViewCountOutput)
+        {
+            return XrResult::XR_ERROR_SIZE_INSUFFICIENT;
+        }
+        
+        //TODO make this more real???
+        ViewState->viewStateFlags = (XR_SPACE_LOCATION_POSITION_VALID_BIT | XR_SPACE_LOCATION_POSITION_TRACKED_BIT);
+        
+        for (int Index = 0; Index < NumViews; ++Index)
+        {
+            if (IsInGameThread())
+            {
+                const FLocateViewInfo& ViewInfoCache = GetLocateViewInfo_GameThread();
+                
+                // Adjust the camera transform for the current eye position.
+                simd_float4x4 CameraTransform = simd_mul(PipelinedFrameStateGame.HeadTransform, ViewInfoCache.ViewTransforms[Index]);
+                
+                Views[Index].pose = OXRVisionOS::ToXrPose(CameraTransform);
+                Views[Index].fov = ViewInfoCache.HmdFovs[Index];
+            }
+            else
+            {
+                check(IsInRenderingThread());
+                check(FrameState.SwiftDrawable);
+                
+                cp_view_t View = cp_drawable_get_view(FrameState.SwiftDrawable, Index);
+                simd_float4 Tangents = cp_view_get_tangents(View);
+                //simd_float2 DepthRange = cp_drawable_get_depth_range(FrameState.SwiftDrawable);  //TODO ??? do we need this?  There is an openxrextension for it, probably should use that.
+    //                                                                          DepthRange[1], /* nearZ */
+    //                                                                          DepthRange[0], /* farZ */
+    //                                                                          true); /* reverseZ */
+                
+                FrameState.HmdFovs[Index].angleLeft     = -FMath::Atan(Tangents[0]);
+                FrameState.HmdFovs[Index].angleRight    =  FMath::Atan(Tangents[1]);
+                FrameState.HmdFovs[Index].angleUp       =  FMath::Atan(Tangents[2]);
+                FrameState.HmdFovs[Index].angleDown     = -FMath::Atan(Tangents[3]);
+
+                // Adjust the camera transform for the current eye position.
+                simd_float4x4 ViewTransform = cp_view_get_transform(View);
+                simd_float4x4 CameraTransform = simd_mul(PipelinedFrameStateRHI.HeadTransform, ViewTransform);
+                
+                Views[Index].pose = OXRVisionOS::ToXrPose(CameraTransform);
+                Views[Index].fov = FrameState.HmdFovs[Index];
+                
+                // Cache the render thread data for the game thread
+                FLocateViewInfo& ViewInfoCache = GetLocateViewInfo_RenderThread();
+                ViewInfoCache.ViewTransforms[Index] = ViewTransform;
+                ViewInfoCache.HmdFovs[Index] = FrameState.HmdFovs[Index];
+            }
+        }
+
+		//StoreHMDPoseOffsetData(hmdInfo.devicePose);
+	}
+
+//	if (IsInGameThread())
+//	{
+//		ENQUEUE_RENDER_COMMAND(UpdateHMDPoseFromGT)([this, NewRenderingFrameState = PipelinedFrameStateGame](FRHICommandListImmediate& RHICmdList)
+//		{
+//			PipelinedFrameStateRendering.HMDPoseInTrackerSpace = NewRenderingFrameState.HMDPoseInTrackerSpace;
+//			RHICmdList.EnqueueLambda([this, NewRHIFrameState = PipelinedFrameStateRendering](FRHICommandListImmediate&)
+//			{
+//				PipelinedFrameStateRHI.HMDPoseInTrackerSpace = NewRHIFrameState.HMDPoseInTrackerSpace;
+//			});
+//		});
+//	}
+//	else
+//	{
+//		check(IsInRenderingThread());
+//		FRHICommandListImmediate& RHICmdList = GetImmediateCommandList_ForRenderCommand();
+//		RHICmdList.EnqueueLambda([this, NewRHIFrameState = PipelinedFrameStateRendering](FRHICommandListImmediate&)
+//		{
+//			PipelinedFrameStateRHI.HMDPoseInTrackerSpace = NewRHIFrameState.HMDPoseInTrackerSpace;
+//		});
+//	}
+
+	return XrResult::XR_SUCCESS;
+}
+
+// void FOXRVisionOSSession::StoreHMDPoseOffsetData(const PoseData& HmdPose)
+// {
+// 	FTransform HmdTransform = OXRVisionOS::ToFTransform(HmdPose);
+// 	// After a tracking disruption hmd transform often becomes slightly denormalized, more than our checkslow threshold.
+// 	HmdTransform.NormalizeRotation();
+
+// 	FPipelinedFrameState& FrameStateForThread = GetPipelinedFrameStateForThread();
+// 	FrameStateForThread.HMDPoseInTrackerSpace = HmdTransform;
+// }
+
+void FOXRVisionOSSession::GetHMDTransform(XrTime DisplayTime, FTransform& OutTransform, XrSpaceLocationFlags& OutLocationFlags)
+{
+    //TODO just ignoring XrTime right now!!!
+    
+    FPipelinedFrameState& FrameState = GetPipelinedFrameStateForThread();
+    
+    OutLocationFlags = XR_SPACE_LOCATION_ORIENTATION_VALID_BIT
+        | XR_SPACE_LOCATION_POSITION_VALID_BIT
+        | XR_SPACE_LOCATION_ORIENTATION_TRACKED_BIT
+    | XR_SPACE_LOCATION_POSITION_TRACKED_BIT;
+    
+    OutTransform = OXRVisionOS::ToFTransform(FrameState.HeadTransform);
+}
+
+void FOXRVisionOSSession::GetTransformForButton(EOXRVisionOSControllerButton Button, XrTime DisplayTime, FTransform& OutTransform, XrSpaceLocationFlags& OutLocationFlags, FVector& OutLinearVelocity, FVector& OutAngularVelocity, XrSpaceVelocityFlags& OutVelocityFlags)
+{
+	OutLocationFlags = 0;
+	OutVelocityFlags = 0;
+
+//	bool bGaze = false;
+	int32 DeviceHandle = INVALID_DEVICE_HANDLE;
+	switch (Button)
+	{
+	case EOXRVisionOSControllerButton::GripL:
+	case EOXRVisionOSControllerButton::AimL:
+		DeviceHandle = Controllers->GetLeftControllerHandle();
+		break;
+	case EOXRVisionOSControllerButton::GripR:
+	case EOXRVisionOSControllerButton::AimR:
+		DeviceHandle = Controllers->GetRightControllerHandle();
+		break;
+	case EOXRVisionOSControllerButton::HMDPose:
+        check(false); // should be using GetHMDTransform
+		DeviceHandle = HMDHandle;
+		break;
+	// case EOXRVisionOSControllerButton::GazePose:
+	// 	DeviceHandle = HMDHandle;
+	// 	bGaze = true;
+	// 	break;
+	default:
+		break;
+	}
+	if (DeviceHandle == INVALID_DEVICE_HANDLE)
+	{
+		OutTransform = FTransform::Identity;
+		return;
+	}
+
+	// if (bGaze)
+	// {
+	// 	// We get the gaze origin from the hmd.
+	// 	bool bSuccess = Tracker->GetGazeResult(DeviceHandle, DisplayTime, FrameNumber, GazeResultData, HMDResultData);
+	// 	if (!bSuccess)
+	// 	{
+	// 		OutTransform = FTransform::Identity;
+	// 		return;
+	// 	}
+	// 	OutLocationFlags = OXRVisionOS::ToXrSpaceLocationFlags(HMDResultData);
+	// 	OutVelocityFlags = 0; // gaze never provides velocity data
+
+	// 	// gazeDirection is in gazeOriginPose space.
+	// 	PoseData& HMDGazeOriginPose = HMDResultData.Hmd.GazeOriginPose;
+	// 	FQuat GazeOriginQuat =
+	// 	FVector GazeOriginVector =
+	// 	FQuat GazeQuat =
+
+	// 	OutTransform = FTransform(GazeQuat) * FTransform(GazeOriginQuat, GazeOriginVector);
+	// }
+	// else
+	{
+		OXRVisionOS::TrackerResultData ResultData = {};
+		bool bSuccess = Tracker->GetResult(DeviceHandle, DisplayTime, ResultData);
+		if (!bSuccess)
+		{
+			OutTransform = FTransform::Identity;
+			return;
+		}
+
+		//OutLocationFlags = OXRVisionOS::ToXrSpaceLocationFlags(ResultData);
+		//OutVelocityFlags = OXRVisionOS::ToXrSpaceVelocityFlags(ResultData);
+
+		if ((OutLocationFlags & (XR_SPACE_LOCATION_POSITION_VALID_BIT | XR_SPACE_LOCATION_ORIENTATION_VALID_BIT)))
+		{
+			//PoseData& Pose = DeviceHandle == HMDHandle ? hmdInfo.Pose : vrControllerInfo.Pose;
+			OXRVisionOS::TrackerPoseData Pose; //TEMP
+
+			// If part of the transform is invalid send along identity/zero.
+			FQuat Quat = (OutLocationFlags & XR_SPACE_LOCATION_ORIENTATION_VALID_BIT) ? OXRVisionOS::ToFQuat(Pose.orientation) : FQuat::Identity;
+			FVector Vector = (OutLocationFlags & XR_SPACE_LOCATION_POSITION_VALID_BIT) ? OXRVisionOS::ToFVector(Pose.position) : FVector(0.0f, 0.0f, 0.0f);
+
+			OutTransform = FTransform(Quat, Vector);
+
+			// api provides just one controller pose, which is a reasonable aim pose, we shall make both adjustible.
+			if (Button == EOXRVisionOSControllerButton::GripL || Button == EOXRVisionOSControllerButton::GripR)
+			{
+				OutTransform.SetRotation(OutTransform.TransformRotation(GripPoseAdjustment));
+			}
+			else if (Button == EOXRVisionOSControllerButton::AimL || Button == EOXRVisionOSControllerButton::AimR)
+			{
+				OutTransform.SetRotation(OutTransform.TransformRotation(AimPoseAdjustment));
+			}
+		}
+
+		// if (OutVelocityFlags & XR_SPACE_VELOCITY_LINEAR_VALID_BIT)
+		// {
+		// 	OutLinearVelocity = OXRVisionOS::ToFVector(ResultData.velocity);
+		// }
+		// if (OutVelocityFlags & XR_SPACE_VELOCITY_ANGULAR_VALID_BIT)
+		// {
+		// 	OutAngularVelocity = OXRVisionOS::ToFVector(ResultData.angularVelocity);
+		// }
+	}
+}
+
+// FOpenXRInput::Tick, game thread, after xrSyncActions
+void FOXRVisionOSSession::OnPostSyncActions()
+{
+}
+
+XrTime FOXRVisionOSSession::GetCurrentTime() const
+{	
+    const CFAbsoluteTime CurrentTime = CFAbsoluteTimeGetCurrent();
+	return OXRVisionOS::APITimeToXrTime(CurrentTime);
+}
+
+void FOXRVisionOSSession::SetSessionState(XrSessionState NewState)
+{
+	UE_LOG(LogOXRVisionOS, Log, TEXT("FOXRVisionOSSession %s -> %s"), OpenXRSessionStateToString(SessionState), OpenXRSessionStateToString(NewState));
+
+	SessionState = NewState;
+
+	XrEventDataSessionStateChanged Event;
+	Event.type = XR_TYPE_EVENT_DATA_SESSION_STATE_CHANGED;
+	Event.next = nullptr;
+	Event.session = (XrSession)this;
+	Event.state = SessionState;
+	Event.time = GetCurrentTime();
+
+	Instance->EnqueueEvent(Event);
+}
+
+bool FOXRVisionOSSession::IsSessionRunning() const
+{
+	return bIsRunning;
+}
+
+const FOXRVisionOSSession::FPipelinedFrameState& FOXRVisionOSSession::GetPipelinedFrameStateForThread() const
+{
+	if (IsInRHIThread())
+	{
+		return PipelinedFrameStateRHI;
+	}
+	else if (IsInRenderingThread())
+	{
+		return PipelinedFrameStateRendering;
+	}
+	else
+	{
+		check(IsInGameThread());
+		return PipelinedFrameStateGame;
+	}
+}
+
+FOXRVisionOSSession::FPipelinedFrameState& FOXRVisionOSSession::GetPipelinedFrameStateForThread()
+{
+	return const_cast<FOXRVisionOSSession::FPipelinedFrameState&>(static_cast<const FOXRVisionOSSession&>(*this).GetPipelinedFrameStateForThread());
+}
diff --git a/Engine/Platforms/VisionOS/Plugins/Runtime/OpenXRVisionOS/Source/OXRVisionOS/Private/OXRVisionOSSession.h b/Engine/Platforms/VisionOS/Plugins/Runtime/OpenXRVisionOS/Source/OXRVisionOS/Private/OXRVisionOSSession.h
new file mode 100644
--- /dev/null
+++ b/Engine/Platforms/VisionOS/Plugins/Runtime/OpenXRVisionOS/Source/OXRVisionOS/Private/OXRVisionOSSession.h
@@ -0,0 +1,241 @@
+// Copyright Epic Games, Inc. All Rights Reserved.
+
+#pragma once
+
+#include "OXRVisionOS.h"
+
+#include "Engine/EngineBaseTypes.h"
+#include "Epic_openxr.h"
+#include "OXRVisionOS_openxr_platform.h"
+#include "HAL/Runnable.h"
+#include "HeadMountedDisplayTypes.h"
+
+#if PLATFORM_VISIONOS
+#import <ARKit/ARKit.h>
+#import <CompositorServices/CompositorServices.h>
+#endif
+
+
+class FOXRVisionOSInstance;
+class FOXRVisionOSActionSet;
+class FOXRVisionOSSpace;
+class FOXRVisionOSSwapchain;
+class FOXRVisionOSTracker;
+class FOXRVisionOSController;
+
+class FRHICommandListImmediate;
+enum class EOXRVisionOSControllerButton : int32;
+
+class FOXRVisionOSSession
+{
+public:
+	static XrResult Create(TSharedPtr<FOXRVisionOSSession, ESPMode::ThreadSafe>& OutSession, const XrSessionCreateInfo* createInfo, FOXRVisionOSInstance* Instance);
+	FOXRVisionOSSession(const XrSessionCreateInfo* createInfo, FOXRVisionOSInstance* Instance);
+	~FOXRVisionOSSession();
+	XrResult XrDestroySession();
+
+	XrResult XrEnumerateReferenceSpaces(
+		uint32_t spaceCapacityInput,
+		uint32_t* spaceCountOutput,
+		XrReferenceSpaceType* spaces);
+
+	XrResult XrCreateReferenceSpace(
+		const XrReferenceSpaceCreateInfo* createInfo,
+		XrSpace* space);
+
+	XrResult XrGetReferenceSpaceBoundsRect(
+		XrReferenceSpaceType referenceSpaceType,
+		XrExtent2Df* bounds);
+
+	XrResult XrCreateActionSpace(
+		const XrActionSpaceCreateInfo* createInfo,
+		XrSpace* space);
+
+	XrResult DestroySpace(
+		FOXRVisionOSSpace* Space);
+
+	XrResult XrEnumerateSwapchainFormats(
+		uint32_t formatCapacityInput,
+		uint32_t* formatCountOutput,
+		int64_t* formats);
+
+	XrResult XrCreateSwapchain(
+		const XrSwapchainCreateInfo* createInfo,
+		XrSwapchain* swapchain);
+
+	XrResult DestroySwapchain(
+		FOXRVisionOSSwapchain* Swapchain);
+
+	XrResult XrAttachSessionActionSets(
+		const XrSessionActionSetsAttachInfo* attachInfo);
+
+	XrResult XrGetCurrentInteractionProfile(
+		XrPath topLevelUserPath,
+		XrInteractionProfileState* interactionProfile);
+
+	XrResult XrGetActionStateBoolean(
+		const XrActionStateGetInfo* getInfo,
+		XrActionStateBoolean* state);
+
+	XrResult XrGetActionStateFloat(
+		const XrActionStateGetInfo* getInfo,
+		XrActionStateFloat* state);
+
+	XrResult XrGetActionStateVector2f(
+		const XrActionStateGetInfo* getInfo,
+		XrActionStateVector2f* state);
+
+	XrResult XrGetActionStatePose(
+		const XrActionStateGetInfo* getInfo,
+		XrActionStatePose* state);
+
+	XrResult XrSyncActions(
+		const XrActionsSyncInfo* syncInfo);
+
+	XrResult XrEnumerateBoundSourcesForAction(
+		const XrBoundSourcesForActionEnumerateInfo* enumerateInfo,
+		uint32_t sourceCapacityInput,
+		uint32_t* sourceCountOutput,
+		XrPath* sources);
+
+	XrResult XrApplyHapticFeedback(
+		const XrHapticActionInfo* hapticActionInfo,
+		const XrHapticBaseHeader* hapticFeedback);
+
+	XrResult XrStopHapticFeedback(
+		const XrHapticActionInfo* hapticActionInfo);
+
+	XrResult XrBeginSession(
+		const XrSessionBeginInfo* beginInfo);
+
+	XrResult XrEndSession();
+
+	void EndSessionInternal();
+
+	XrResult XrRequestExitSession();
+
+	XrResult XrWaitFrame(
+		const XrFrameWaitInfo* frameWaitInfo,
+		XrFrameState* frameState);
+
+	void OnBeginRendering_GameThread();
+
+	void OnBeginRendering_RenderThread();
+
+	XrResult XrBeginFrame(
+		const XrFrameBeginInfo* frameBeginInfo);
+
+	XrResult XrEndFrame(
+		const XrFrameEndInfo* frameEndInfo);
+
+	XrResult XrLocateViews(
+		const XrViewLocateInfo* viewLocateInfo,
+		XrViewState* viewState,
+		uint32_t viewCapacityInput,
+		uint32_t* viewCountOutput,
+		XrView* views);
+
+	XrTime GetCurrentTime() const;
+
+	void SetSessionState(XrSessionState NewState);
+	XrSessionState GetSessionState() const { return SessionState; }
+
+	bool IsSessionRunning() const;
+
+	FOXRVisionOSTracker& GetTrackerChecked() const { check(Tracker.IsValid()); return *Tracker; }
+
+	struct FPipelinedFrameState
+	{
+		int32 FrameCounter = -1;
+		XrTime PredictedDisplayTime = 0;
+		IRHICommandContext* CommandListContext = nullptr;
+		//FTransform HMDPoseInTrackerSpace = FTransform::Identity;
+		//PoseData HmdPose = {};
+		cp_frame_t SwiftFrame = nullptr;
+        cp_frame_timing_t SwiftFrameTiming = nullptr;
+        simd_float4x4 HeadTransform = matrix_identity_float4x4;
+        cp_drawable_t SwiftDrawable = nullptr;
+        cp_frame_timing_t SwiftFinalFrameTiming = nullptr; // TODO should we do this or should we overwrite SwifFrameTiming???
+        ar_device_anchor_t DeviceAnchor = nullptr;
+		bool bSynchronizing = true; // when true we are in the process of starting up a new session and the new frame state has not yet propagated.
+		XrFovf HmdFovs[2];
+        int32 LocateViewInfoBufferIndex = 0;
+	};
+	const FPipelinedFrameState& GetPipelinedFrameStateForThread() const;
+	FPipelinedFrameState& GetPipelinedFrameStateForThread();
+	void GetHMDTransform(XrTime DisplayTime, FTransform& OutTransform, XrSpaceLocationFlags& OutLocationFlags);
+	void GetTransformForButton(EOXRVisionOSControllerButton Button, XrTime DisplayTime, FTransform& OutTransform, XrSpaceLocationFlags& OutLocationFlags, FVector& OutLinearVelocity, FVector& OutAngularVelocity, XrSpaceVelocityFlags& OutVelocityFlags);
+	void SetTriggerEffectVRControllers(int32 ControllerId, const FInputDeviceProperty* Property);
+	// FOpenXRInput::Tick, game thread, after xrSyncActions
+	void OnPostSyncActions(); 
+
+	struct FLocateViewInfo
+    {
+        uint32_t NumViews = 0;
+        simd_float4x4 ViewTransforms[2];
+        XrFovf HmdFovs[2];
+    };
+    static constexpr int LocateViewInfoBufferLength = 3;
+    FLocateViewInfo LocateViewInfoBuffer[LocateViewInfoBufferLength];
+    void LocateViewInfoIndexAdvance() { ++PipelinedFrameStateGame.LocateViewInfoBufferIndex; }
+    const FLocateViewInfo& GetLocateViewInfo_GameThread() { return LocateViewInfoBuffer[(PipelinedFrameStateGame.LocateViewInfoBufferIndex - 1) % LocateViewInfoBufferLength]; }
+    FLocateViewInfo& GetLocateViewInfo_RenderThread() { return LocateViewInfoBuffer[PipelinedFrameStateRendering.LocateViewInfoBufferIndex % LocateViewInfoBufferLength]; }
+
+private:
+	bool bCreateFailed = false;
+	bool bIsRunning = false;
+	bool bExitRequested = false;
+	FOXRVisionOSInstance* Instance = nullptr;
+	uint64_t XRTimeScale = 1;
+	XrSessionState SessionState = XR_SESSION_STATE_UNKNOWN;
+	bool ActionSetsAlreadyAttached = false;
+	TArray<TSharedPtr<FOXRVisionOSActionSet, ESPMode::ThreadSafe>> ActionSets;
+	TSharedPtr<FOXRVisionOSController, ESPMode::ThreadSafe> Controllers;
+	TArray<TSharedPtr<FOXRVisionOSSpace, ESPMode::ThreadSafe>> Spaces;
+	FQuat GripPoseAdjustment;
+	FQuat AimPoseAdjustment;
+
+	FPlatformMemory::FPlatformVirtualMemoryBlock m_reprojectionData;
+	FPlatformMemory::FPlatformVirtualMemoryBlock m_displayData;
+
+
+	TArray<TSharedPtr<FOXRVisionOSSwapchain, ESPMode::ThreadSafe>> Swapchains;
+
+	int32 NextBackBufferIndex = 0;
+	int32 PreviousBackBufferIndexInBegin = -1;
+	int32 SessionFrameCounter = 0;
+
+	// Because OpenXR works with view poses, but works with hmd pose we cache an hmd pose when we get the trackign data.
+	// If the OpenXR Application modifies the view poses such that they are not relative to the hmd pose this will fail and the altered view poses will not behave as expected.
+	// This should work correctly if api modifies the interpupillary distance, provided that is properly pipelined on the api side (meaning it is constant for each fliparg).
+	//void StoreHMDPoseOffsetData(const PoseData& HmdPose);
+
+	static const int32 INVALID_DEVICE_HANDLE = -1;
+	int32 HMDHandle = INVALID_DEVICE_HANDLE;
+
+	TSharedPtr <FOXRVisionOSTracker, ESPMode::ThreadSafe> Tracker;
+
+	//TEMP
+	static const int32 RESULT_OK = 0;
+
+
+	void OnWorldTickStart(UWorld* World, ELevelTick TickType, float DeltaTime);
+	FDelegateHandle OnWorldTickStartDelegateHandle;
+
+	void WaitUntil();
+    void StartSubmission();
+	FEvent* XrWaitFrameEvent = nullptr;
+	FPipelinedFrameState	PipelinedFrameStateGame;
+	FPipelinedFrameState	PipelinedFrameStateRendering;
+	FPipelinedFrameState	PipelinedFrameStateRHI;
+	int32 CachedBeginFlipFrameCounter = INT32_MAX;
+
+	// TODO Reorganize this class?
+
+	FOpenXRRenderBridge* RenderBridge = nullptr;
+
+	CP_OBJECT_cp_layer_renderer* VOSLayerRenderer = nullptr;
+	ar_session_t ARKitSession = nullptr;
+	ar_world_tracking_provider_t ARKitWorldTrackingProvider = nullptr;
+	ar_device_anchor_t ARKitHMDAnchor = nullptr;
+};
diff --git a/Engine/Platforms/VisionOS/Plugins/Runtime/OpenXRVisionOS/Source/OXRVisionOS/Private/OXRVisionOSSpace.cpp b/Engine/Platforms/VisionOS/Plugins/Runtime/OpenXRVisionOS/Source/OXRVisionOS/Private/OXRVisionOSSpace.cpp
new file mode 100644
--- /dev/null
+++ b/Engine/Platforms/VisionOS/Plugins/Runtime/OpenXRVisionOS/Source/OXRVisionOS/Private/OXRVisionOSSpace.cpp
@@ -0,0 +1,329 @@
+// Copyright Epic Games, Inc. All Rights Reserved.
+
+#include "OXRVisionOSSpace.h"
+#include "OXRVisionOSSession.h"
+#include "OXRVisionOSPlatformUtils.h"
+#include "OXRVisionOSTracker.h"
+#include "OXRVisionOSAction.h"
+#include "OXRVisionOSController.h"
+
+//#include "../../../../../../../../Source/ThirdParty/OpenXR/include/openxr/openxr_platform.h"
+//#include <openxr/openxr_platform.h>
+
+XrResult FOXRVisionOSSpace::CreateActionSpace(TSharedPtr<FOXRVisionOSSpace, ESPMode::ThreadSafe>& OutSpace, const XrActionSpaceCreateInfo* CreateInfo, FOXRVisionOSSession* Session)
+{
+	if (CreateInfo == nullptr || Session == nullptr)
+	{
+		return XrResult::XR_ERROR_VALIDATION_FAILURE;
+	}
+
+	if (CreateInfo->type != XR_TYPE_ACTION_SPACE_CREATE_INFO)
+	{
+		return XrResult::XR_ERROR_VALIDATION_FAILURE;
+	}
+
+	if (CreateInfo->action == XR_NULL_HANDLE)
+	{
+		return XrResult::XR_ERROR_HANDLE_INVALID;
+	}
+	
+	const FOXRVisionOSAction& Action = *((FOXRVisionOSAction*)CreateInfo->action);
+	if (Action.GetActionType() != XR_ACTION_TYPE_POSE_INPUT)
+	{
+		return XR_ERROR_ACTION_TYPE_MISMATCH;
+	}
+
+	OutSpace = MakeShared<FOXRVisionOSSpace, ESPMode::ThreadSafe>(CreateInfo, Session);
+	if (OutSpace->bCreateFailed)
+	{
+		OutSpace = nullptr;
+		return XrResult::XR_ERROR_RUNTIME_FAILURE;
+	}
+
+	return XrResult::XR_SUCCESS;
+}
+
+FOXRVisionOSSpace::FOXRVisionOSSpace(const XrActionSpaceCreateInfo* CreateInfo, FOXRVisionOSSession* InSession)
+{
+	Session = InSession;
+
+	bIsReferenceSpace = false;
+	Action = CreateInfo->action;
+	SubactionPath = CreateInfo->subactionPath;
+	Pose = CreateInfo->poseInActionSpace;
+}
+
+XrResult FOXRVisionOSSpace::CreateReferenceSpace(TSharedPtr<FOXRVisionOSSpace, ESPMode::ThreadSafe>& OutSpace, const XrReferenceSpaceCreateInfo* CreateInfo, FOXRVisionOSSession* Session)
+{
+	if (CreateInfo == nullptr || Session == nullptr)
+	{
+		return XrResult::XR_ERROR_VALIDATION_FAILURE;
+	}
+
+	if (CreateInfo->type != XR_TYPE_REFERENCE_SPACE_CREATE_INFO)
+	{
+		return XrResult::XR_ERROR_VALIDATION_FAILURE;
+	}
+
+	OutSpace = MakeShared<FOXRVisionOSSpace, ESPMode::ThreadSafe>(CreateInfo, Session);
+	if (OutSpace->bCreateFailed)
+	{
+		OutSpace = nullptr;
+		return XrResult::XR_ERROR_RUNTIME_FAILURE;
+	}
+
+	return XrResult::XR_SUCCESS;
+}
+
+FOXRVisionOSSpace::FOXRVisionOSSpace(const XrReferenceSpaceCreateInfo* CreateInfo, FOXRVisionOSSession* InSession)
+{
+	Session = InSession;
+
+	bIsReferenceSpace = true;
+	ReferenceSpaceType = CreateInfo->referenceSpaceType;
+	Pose = CreateInfo->poseInReferenceSpace;
+	UEPose = ToFTransform(Pose);
+	UEPoseInverse = UEPose.Inverse();
+}
+
+FOXRVisionOSSpace::~FOXRVisionOSSpace()
+{
+	if (bCreateFailed)
+	{
+		UE_LOG(LogOXRVisionOS, Warning, TEXT("Destructing FOXRVisionOSSpace because create failed."));
+	}
+}
+
+XrResult FOXRVisionOSSpace::XrDestroySpace()
+{
+	// Session->DestroySpace() can delete this, so better just return after that.
+	return Session->DestroySpace(this);
+}
+
+XrResult FOXRVisionOSSpace::XrLocateSpace(
+	XrSpace                                     InBaseSpace,
+	XrTime                                      Time,
+	XrSpaceLocation*							Location) 
+{
+	if (InBaseSpace == XR_NULL_HANDLE)
+	{
+		return XR_ERROR_HANDLE_INVALID;
+	}
+
+	XrSpaceVelocity* Velocity = OpenXR::FindChainedStructByType<XrSpaceVelocity>(Location->next, (XrStructureType)XR_TYPE_SPACE_VELOCITY);
+
+	// When not focused we do not return any tracking data.
+	if (Session->GetSessionState() != XrSessionState::XR_SESSION_STATE_FOCUSED)
+	{
+		Location->locationFlags = 0;
+		if (Velocity)
+		{
+			Velocity->velocityFlags = 0;
+		}
+
+		return XrResult::XR_SUCCESS;
+	}
+
+	if (InBaseSpace == (XrSpace)this)
+	{
+		// locate vs self is always identity
+		Location->locationFlags = XR_SPACE_LOCATION_ORIENTATION_VALID_BIT
+			| XR_SPACE_LOCATION_POSITION_VALID_BIT
+			| XR_SPACE_LOCATION_ORIENTATION_TRACKED_BIT
+			| XR_SPACE_LOCATION_POSITION_TRACKED_BIT;
+		Location->pose = XrPosefIdentity;
+
+		if (Velocity)
+		{
+			// velocity vs self is always zero
+			Velocity->velocityFlags = XR_SPACE_VELOCITY_LINEAR_VALID_BIT | XR_SPACE_VELOCITY_ANGULAR_VALID_BIT;
+			Velocity->linearVelocity = OXRVisionOS::XrZeroVector;
+			Velocity->angularVelocity = OXRVisionOS::XrZeroVector;
+		}
+
+	}
+
+	FOXRVisionOSSpace& BaseSpace = *(FOXRVisionOSSpace*)InBaseSpace;
+
+	if (IsSpaceRelativeToSameEntity(BaseSpace))
+	{
+		// locate vs a space which is relative to the same entity is always valid and tracked and determined only by their poses relative to that entity.
+		Location->locationFlags = XR_SPACE_LOCATION_ORIENTATION_VALID_BIT
+			| XR_SPACE_LOCATION_POSITION_VALID_BIT
+			| XR_SPACE_LOCATION_ORIENTATION_TRACKED_BIT
+			| XR_SPACE_LOCATION_POSITION_TRACKED_BIT;
+
+		FTransform Result = UEPose * BaseSpace.UEPoseInverse;
+		Location->pose = ToXrPose(Result);
+
+		if (Velocity)
+		{
+			// velocity vs a sibling attached to the same parent is always zero
+			Velocity->velocityFlags = XR_SPACE_VELOCITY_LINEAR_VALID_BIT | XR_SPACE_VELOCITY_ANGULAR_VALID_BIT;
+			Velocity->linearVelocity = OXRVisionOS::XrZeroVector;
+			Velocity->angularVelocity = OXRVisionOS::XrZeroVector;
+		}
+
+		return XrResult::XR_SUCCESS;
+	}
+
+	FTransform MyTransform;
+	XrSpaceLocationFlags MyFlags = 0;
+	FVector MyLinearVelocity;
+	FVector MyAngularVelocity;
+	XrSpaceVelocityFlags MyVelocityFlags = 0;
+
+	FTransform TheirInverseTransform;
+	XrSpaceLocationFlags TheirFlags = 0;
+	FVector TheirInverseLinearVelocity;
+	FVector TheirInverseAngularVelocity;
+	XrSpaceVelocityFlags TheirVelocityFlags = 0;
+
+	GetTransform(Time, MyTransform, MyFlags, MyLinearVelocity, MyAngularVelocity, MyVelocityFlags);
+	BaseSpace.GetInverseTransform(Time, TheirInverseTransform, TheirFlags, TheirInverseLinearVelocity, TheirInverseAngularVelocity, TheirVelocityFlags);
+	FTransform Result = MyTransform * TheirInverseTransform;
+	Location->pose = ToXrPose(Result);
+
+	Location->locationFlags = MyFlags & TheirFlags;
+
+	if (Velocity)
+	{
+		Velocity->velocityFlags = 0;
+		if ((Location->locationFlags & XR_SPACE_LOCATION_ORIENTATION_VALID_BIT) && (MyVelocityFlags & XR_SPACE_VELOCITY_LINEAR_VALID_BIT) && (TheirVelocityFlags & XR_SPACE_VELOCITY_LINEAR_VALID_BIT))
+		{
+			FVector RelativeVelocity = MyLinearVelocity + TheirInverseLinearVelocity;
+			Velocity->linearVelocity = OXRVisionOS::ToXrVector3f(TheirInverseTransform.TransformVectorNoScale(RelativeVelocity));
+			Velocity->velocityFlags |= XR_SPACE_VELOCITY_LINEAR_VALID_BIT;
+		}
+		if ((Location->locationFlags & XR_SPACE_LOCATION_ORIENTATION_VALID_BIT) && (MyVelocityFlags & XR_SPACE_VELOCITY_ANGULAR_VALID_BIT) && (TheirVelocityFlags & XR_SPACE_VELOCITY_ANGULAR_VALID_BIT))
+		{
+			FVector RelativeAngularVelocity = MyAngularVelocity + TheirInverseAngularVelocity;
+			Velocity->angularVelocity = OXRVisionOS::ToXrVector3f(TheirInverseTransform.TransformVectorNoScale(RelativeAngularVelocity));
+			Velocity->velocityFlags |= XR_SPACE_VELOCITY_ANGULAR_VALID_BIT;
+		}
+	}
+
+	return XrResult::XR_SUCCESS;
+}
+
+bool FOXRVisionOSSpace::IsSpaceRelativeToSameEntity(FOXRVisionOSSpace& OtherSpace) const
+{
+	if (bIsReferenceSpace != OtherSpace.bIsReferenceSpace)
+	{
+		return false;
+	}
+
+	if (bIsReferenceSpace)
+	{
+		return ReferenceSpaceType == OtherSpace.ReferenceSpaceType;
+	}
+	else
+	{
+		//TODO is this correct??
+		return Action == OtherSpace.Action;
+	}
+}
+void FOXRVisionOSSpace::GetTransform(XrTime DisplayTime, FTransform& OutTransform, XrSpaceLocationFlags& OutLocationFlags, FVector& OutLinearVelocity, FVector& OutAngularVelocity, XrSpaceVelocityFlags& OutVelocityFlags)
+{
+	OutLocationFlags = 0;
+
+	if (bIsReferenceSpace)
+	{
+		switch (ReferenceSpaceType)
+		{
+		case XR_REFERENCE_SPACE_TYPE_VIEW:
+		{
+			// The view is the hmd.
+            // Local space is tracking space, so it is always at the origin + any pose it was given and velocity zero.
+            OutLocationFlags = XR_SPACE_LOCATION_ORIENTATION_VALID_BIT
+                | XR_SPACE_LOCATION_POSITION_VALID_BIT
+                | XR_SPACE_LOCATION_ORIENTATION_TRACKED_BIT
+                | XR_SPACE_LOCATION_POSITION_TRACKED_BIT;
+            Session->GetHMDTransform(DisplayTime, OutTransform, OutLocationFlags);
+            OutVelocityFlags = 0;
+            OutLinearVelocity = FVector::ZeroVector;
+            OutAngularVelocity = FVector::ZeroVector;
+            return;
+		}
+		case XR_REFERENCE_SPACE_TYPE_LOCAL:
+		{
+			// Local space is tracking space, so it is always at the origin + any pose it was given and velocity zero.
+			OutLocationFlags = XR_SPACE_LOCATION_ORIENTATION_VALID_BIT
+				| XR_SPACE_LOCATION_POSITION_VALID_BIT
+				| XR_SPACE_LOCATION_ORIENTATION_TRACKED_BIT
+				| XR_SPACE_LOCATION_POSITION_TRACKED_BIT;
+			OutTransform = UEPose;
+			OutVelocityFlags = XR_SPACE_VELOCITY_LINEAR_VALID_BIT | XR_SPACE_VELOCITY_ANGULAR_VALID_BIT;
+			OutLinearVelocity = FVector::ZeroVector;
+			OutAngularVelocity = FVector::ZeroVector;
+			return;
+		}
+//		case XR_REFERENCE_SPACE_TYPE_STAGE:
+//		{
+//			// The stage is the space of the reference space, aka the playable area.  See also XrGetReferenceSpaceBoundsRect.
+//
+//			FOXRVisionOSTracker& Tracker = Session->GetTrackerChecked();
+//			OXRVisionOS::TrackerPoseData PoseData = {};
+//			bool bSuccess = Tracker.GetStagePose(PoseData);
+//			if (bSuccess)
+//			{
+//				OutLocationFlags = XR_SPACE_LOCATION_ORIENTATION_VALID_BIT
+//					| XR_SPACE_LOCATION_POSITION_VALID_BIT
+//					| XR_SPACE_LOCATION_ORIENTATION_TRACKED_BIT
+//					| XR_SPACE_LOCATION_POSITION_TRACKED_BIT;
+//				FQuat OriginQuat = OXRVisionOS::ToFQuat(PoseData.orientation);
+//				FVector OriginVector = OXRVisionOS::ToFVector(PoseData.position);
+//				OutTransform = FTransform(OriginQuat, OriginVector) * UEPose;
+//				OutVelocityFlags = XR_SPACE_VELOCITY_LINEAR_VALID_BIT | XR_SPACE_VELOCITY_ANGULAR_VALID_BIT;
+//				OutLinearVelocity = FVector::ZeroVector;
+//				OutAngularVelocity = FVector::ZeroVector;				
+//			}
+//			return;
+//		}			
+		default:
+			// Shouldn't be possible to create an unsupported reference space.
+			check(false);
+			return;
+		}
+	}
+	else
+	{
+		// Not a reference space, so an action space
+		check(Action != XR_NULL_HANDLE);
+
+		FOXRVisionOSAction* TheAction = (FOXRVisionOSAction*)Action;
+
+		// figure out what transform to get
+		EOXRVisionOSControllerButton Button = TheAction->GetPoseButton(SubactionPath);
+		Session->GetTransformForButton(Button, DisplayTime, OutTransform, OutLocationFlags, OutLinearVelocity, OutAngularVelocity, OutVelocityFlags);
+	}
+}
+
+void FOXRVisionOSSpace::GetInverseTransform(XrTime Time, FTransform& OutTransform, XrSpaceLocationFlags& OutLocationFlags, FVector& OutLinearVelocity, FVector& OutAngularVelocity, XrSpaceVelocityFlags& OutVelocityFlags)
+{
+	if (bIsReferenceSpace && (ReferenceSpaceType == XR_REFERENCE_SPACE_TYPE_LOCAL))
+	{
+		// Local space is tracking space, so it is always at the origin + any pose it was given and velocity zero.
+		OutLocationFlags = XR_SPACE_LOCATION_ORIENTATION_VALID_BIT
+			| XR_SPACE_LOCATION_POSITION_VALID_BIT
+			| XR_SPACE_LOCATION_ORIENTATION_TRACKED_BIT
+			| XR_SPACE_LOCATION_POSITION_TRACKED_BIT;
+		OutTransform = UEPoseInverse;
+		OutVelocityFlags = XR_SPACE_VELOCITY_LINEAR_VALID_BIT | XR_SPACE_VELOCITY_ANGULAR_VALID_BIT;
+		OutLinearVelocity = FVector::ZeroVector;
+		OutAngularVelocity = FVector::ZeroVector;
+	}
+	else
+	{
+		GetTransform(Time, OutTransform, OutLocationFlags, OutLinearVelocity, OutAngularVelocity, OutVelocityFlags);
+		OutTransform = OutTransform.Inverse();
+		if (OutVelocityFlags & XR_SPACE_VELOCITY_LINEAR_VALID_BIT)
+		{
+			OutLinearVelocity = -OutLinearVelocity;
+		}
+		if (OutVelocityFlags & XR_SPACE_VELOCITY_LINEAR_VALID_BIT)
+		{
+			OutAngularVelocity = -OutAngularVelocity;
+		}
+	}
+}
diff --git a/Engine/Platforms/VisionOS/Plugins/Runtime/OpenXRVisionOS/Source/OXRVisionOS/Private/OXRVisionOSSpace.h b/Engine/Platforms/VisionOS/Plugins/Runtime/OpenXRVisionOS/Source/OXRVisionOS/Private/OXRVisionOSSpace.h
new file mode 100644
--- /dev/null
+++ b/Engine/Platforms/VisionOS/Plugins/Runtime/OpenXRVisionOS/Source/OXRVisionOS/Private/OXRVisionOSSpace.h
@@ -0,0 +1,57 @@
+// Copyright Epic Games, Inc. All Rights Reserved.
+
+#pragma once
+
+#include "OXRVisionOS.h"
+#include "OXRVisionOS_openxr_platform.h"
+#include "Epic_openxr.h"
+
+class FOXRVisionOSSession;
+
+class FOXRVisionOSSpace
+{
+public:
+	static XrResult CreateActionSpace(TSharedPtr<FOXRVisionOSSpace, ESPMode::ThreadSafe>& OutSpace, const XrActionSpaceCreateInfo* createInfo, FOXRVisionOSSession* Session);
+	FOXRVisionOSSpace(const XrActionSpaceCreateInfo* createInfo, FOXRVisionOSSession* Session);
+	static XrResult CreateReferenceSpace(TSharedPtr<FOXRVisionOSSpace, ESPMode::ThreadSafe>& OutSpace, const XrReferenceSpaceCreateInfo* createInfo, FOXRVisionOSSession* Session);
+	FOXRVisionOSSpace(const XrReferenceSpaceCreateInfo* createInfo, FOXRVisionOSSession* Session);
+	~FOXRVisionOSSpace();
+	XrResult XrDestroySpace();
+
+	XrResult XrLocateSpace(
+		XrSpace                                     baseSpace,
+		XrTime                                      time,
+		XrSpaceLocation* location);
+	void GetTransform(XrTime Time, FTransform& OutTransform, XrSpaceLocationFlags& OutLocationFlags, FVector& OutLinearVelocity, FVector& OutAngularVelocity, XrSpaceVelocityFlags& OutVelocityFlags);
+	void GetInverseTransform(XrTime DisplayTime, FTransform& OutTransform, XrSpaceLocationFlags& OutLocationFlags, FVector& OutLinearVelocity, FVector& OutAngularVelocity, XrSpaceVelocityFlags& OutVelocityFlags);
+
+	bool IsReferenceSpace() const { return bIsReferenceSpace; }
+	XrReferenceSpaceType GetReferenceSpaceType() const
+	{
+		if (bIsReferenceSpace)
+		{
+			return ReferenceSpaceType;
+		}
+		return XR_REFERENCE_SPACE_TYPE_MAX_ENUM;
+	}
+
+private:
+	const XrPosef XrPosefIdentity = { { 0.0f, 0.0f, 0.0f, 1.0f }, { 0.0f, 0.0f, 0.0f } };
+
+	bool IsSpaceRelativeToSameEntity(FOXRVisionOSSpace& OtherSpace) const;
+
+	bool bCreateFailed = false;
+	FOXRVisionOSSession* Session = nullptr;
+
+	// EITHER a reference space
+	bool bIsReferenceSpace = false;
+	XrReferenceSpaceType ReferenceSpaceType;
+	// OR an action space
+	XrAction Action = XR_NULL_HANDLE;
+	XrPath SubactionPath = XR_NULL_PATH;
+
+	//TODO support recenter past timestamps for 50ms
+	XrPosef Pose;
+	FTransform UEPose;
+	FTransform UEPoseInverse;
+};
diff --git a/Engine/Platforms/VisionOS/Plugins/Runtime/OpenXRVisionOS/Source/OXRVisionOS/Private/OXRVisionOSSwapchain.cpp b/Engine/Platforms/VisionOS/Plugins/Runtime/OpenXRVisionOS/Source/OXRVisionOS/Private/OXRVisionOSSwapchain.cpp
new file mode 100644
--- /dev/null
+++ b/Engine/Platforms/VisionOS/Plugins/Runtime/OpenXRVisionOS/Source/OXRVisionOS/Private/OXRVisionOSSwapchain.cpp
@@ -0,0 +1,268 @@
+// Copyright Epic Games, Inc. All Rights Reserved.
+
+#include "OXRVisionOSSwapChain.h"
+
+
+
+
+#if PLATFORM_MAC
+#include "Mac/CocoaTextView.h"
+@interface FMetalView : FCocoaTextView
+@end
+#endif
+#include "HAL/PlatformFramePacer.h"
+THIRD_PARTY_INCLUDES_START
+#include "mtlpp.hpp"
+THIRD_PARTY_INCLUDES_END
+
+#if PLATFORM_VISIONOS
+#import <CompositorServices/CompositorServices.h>
+#endif
+
+#include "OXRVisionOSSession.h"
+#include "OXRVisionOS_openxr_platform.h"
+#include "RHICommandList.h"
+//#include "MetalRHIOXRVisionOSBridge.h"
+
+// uncomment to make intellisense work better
+//#include "../../../../../../../../Source/ThirdParty/OpenXR/include/openxr/openxr_platform.h"
+
+FOXRVisionOSSwapchain::FSwapchainImage::FSwapchainImage(FTextureRHIRef InImage, EImageState InImageState)
+	: Image(InImage)
+	, ImageState(InImageState)
+{
+
+}
+
+FOXRVisionOSSwapchain::FSwapchainImage::~FSwapchainImage() = default;
+
+XrResult FOXRVisionOSSwapchain::Create(TSharedPtr<FOXRVisionOSSwapchain, ESPMode::ThreadSafe>& OutSwapchain, const XrSwapchainCreateInfo* createInfo, uint32 SwapchainLength, FOXRVisionOSSession* Session)
+{
+	if (createInfo == nullptr || Session == nullptr)
+	{
+		return XrResult::XR_ERROR_VALIDATION_FAILURE;
+	}
+
+	if (createInfo->type != XR_TYPE_SWAPCHAIN_CREATE_INFO)
+	{
+		return XrResult::XR_ERROR_VALIDATION_FAILURE;
+	}
+
+	OutSwapchain = MakeShared<FOXRVisionOSSwapchain, ESPMode::ThreadSafe>(createInfo, SwapchainLength, Session);
+	if (OutSwapchain->CreateFailureResult != XrResult::XR_SUCCESS)
+	{
+		XrResult Ret = OutSwapchain->CreateFailureResult;
+		OutSwapchain = nullptr;
+		return Ret;
+	}
+
+	return XrResult::XR_SUCCESS;
+}
+
+FOXRVisionOSSwapchain::FOXRVisionOSSwapchain(const XrSwapchainCreateInfo* createInfo, uint32 SwapchainLength, FOXRVisionOSSession* InSession)
+{
+	Session = InSession;
+
+	check(createInfo->next == nullptr);
+
+	check(createInfo);
+	CreateInfo = *createInfo;
+
+	{
+		ETextureCreateFlags TextureCreateFlags = ETextureCreateFlags::None;
+		TextureCreateFlags |= TexCreate_RenderTargetable;
+		TextureCreateFlags |= TexCreate_UAV;
+		TextureCreateFlags |= TexCreate_ShaderResource;
+
+		check((CreateInfo.usageFlags & XR_SWAPCHAIN_USAGE_COLOR_ATTACHMENT_BIT) != 0);
+		check((CreateInfo.usageFlags & XR_SWAPCHAIN_USAGE_DEPTH_STENCIL_ATTACHMENT_BIT) == 0);
+		check(CreateInfo.faceCount == 1);
+		check(CreateInfo.arraySize == 1);
+
+		Images.Reserve(SwapchainLength);
+		for (int32 BufferIndex = 0; BufferIndex < SwapchainLength; ++BufferIndex)
+		{
+			// TODO: Add some identifying names?
+			const FRHITextureCreateDesc BackBufferDesc =
+				FRHITextureCreateDesc::Create2D(TEXT("OXRVisionOSSwapchain"), CreateInfo.width, CreateInfo.height, (EPixelFormat)CreateInfo.format)
+				.SetFlags(TextureCreateFlags)
+				.DetermineInititialState();
+				
+			Images.Emplace(RHICreateTexture(BackBufferDesc), EImageState::Released);
+		}
+	}
+}
+
+FOXRVisionOSSwapchain::~FOXRVisionOSSwapchain()
+{
+	if (CreateFailureResult != XrResult::XR_SUCCESS)
+	{
+		UE_LOG(LogOXRVisionOS, Warning, TEXT("Destructing FOXRVisionOSSwapchain because create failed."));
+	}
+
+	//SwapchainWaitManager.Reset();
+}
+
+XrResult FOXRVisionOSSwapchain::XrDestroySwapchain()
+{
+	// Session->DestroySwapchain() can delete this, so better just return after that.
+	return Session->DestroySwapchain(this);
+}
+
+XrResult FOXRVisionOSSwapchain::XrEnumerateSwapchainImages(
+	uint32_t                                    imageCapacityInput,
+	uint32_t*									imageCountOutput,
+	XrSwapchainImageBaseHeader*					images)
+{
+	if (imageCountOutput == nullptr)
+	{
+		return XrResult::XR_ERROR_VALIDATION_FAILURE;
+	}
+	if (imageCapacityInput != 0 && images == nullptr)
+	{
+		return XrResult::XR_ERROR_VALIDATION_FAILURE;
+	}
+
+	*imageCountOutput = Images.Num();
+	if (imageCapacityInput != 0)
+	{
+		if (imageCapacityInput < *imageCountOutput)
+		{
+			return XrResult::XR_ERROR_SIZE_INSUFFICIENT;
+		}
+
+		XrSwapchainImageOXRVisionOS* CastImages = (XrSwapchainImageOXRVisionOS*)images;
+		for (uint32_t i = 0; i < *imageCountOutput; ++i)
+		{
+			XrSwapchainImageOXRVisionOS& Output = CastImages[i];
+			check(Output.type == (XrStructureType)XR_TYPE_SWAPCHAIN_IMAGE_OXRVISIONOS_EPIC)
+			Output.next = nullptr;
+			Output.image = Images[i].Image;
+		}
+	}
+	return XrResult::XR_SUCCESS;
+}
+
+const FOXRVisionOSSwapchain::FSwapchainImage& FOXRVisionOSSwapchain::GetLastReleasedImage() const
+{
+	check(NextImageToRelease < Images.Num());
+
+	// OXR spec says xrEndFrame uses last released image, so we don't really care about the index
+	// just whichever was last
+
+	const FSwapchainImage& ReleasedImage = Images[NextImageToRelease];
+	check(ReleasedImage.ImageState == EImageState::Released);
+
+	return ReleasedImage;
+}
+
+
+XrResult FOXRVisionOSSwapchain::XrAcquireSwapchainImage(
+	const XrSwapchainImageAcquireInfo*			acquireInfo,
+	uint32_t*									index)
+{
+	SCOPED_NAMED_EVENT_TEXT("FOXRVisionOSSwapchain::XrAcquireSwapchainImage", FColor::Turquoise);
+
+	if (index == nullptr)
+	{
+		return XrResult::XR_ERROR_VALIDATION_FAILURE;
+	}
+
+	// acquireInfo can be null
+	if (acquireInfo != nullptr && acquireInfo->type != XR_TYPE_SWAPCHAIN_IMAGE_ACQUIRE_INFO)
+	{
+		return XrResult::XR_ERROR_VALIDATION_FAILURE;
+	}
+
+	FSwapchainImage& NextImage = Images[NextImageToAcquire];
+	if (NextImage.ImageState == EImageState::Released)
+	{
+		NextImage.ImageState = EImageState::Aquired;
+		*index = NextImageToAcquire;
+
+		AcquiredImageQueue.Enqueue(NextImageToAcquire);
+
+		// TODO: Can the swapchain image advance on it's own??
+		// probably not because of tightly coupled images + buffers for videoout
+		NextImageToAcquire = (NextImageToAcquire + 1) % Images.Num();
+
+		return XrResult::XR_SUCCESS;
+	}
+	else
+	{
+		return XrResult::XR_ERROR_CALL_ORDER_INVALID;
+	}
+}
+
+XrResult FOXRVisionOSSwapchain::XrWaitSwapchainImage(
+	const XrSwapchainImageWaitInfo*				waitInfo)
+{
+	SCOPED_NAMED_EVENT_TEXT("FOXRVisionOSSwapchain::XrWaitSwapchainImage", FColor::Turquoise);
+
+	if (waitInfo == nullptr || waitInfo->type != XR_TYPE_SWAPCHAIN_IMAGE_WAIT_INFO)
+	{
+		return XrResult::XR_ERROR_VALIDATION_FAILURE;
+	}
+
+
+	// TODO
+	// We can only wait on one image at a time. Any previously waited on image
+	// has to be released! That _should_ be impossible because OnFinishRendering_RHIThread
+	// should run again before OnBeginRendering_RHIThread is scheduled again. But it could
+	// happen if present isn't called or OnFinish bails early!
+	if (++OutstandingWaits > 1)
+	{
+		// We won't recover because xrReleaseSwapchainImage is scheduled
+		// synchronously after xrWaitSwapchainImage :(
+		// Is there a better return code for this?
+		return XrResult::XR_TIMEOUT_EXPIRED;
+	}
+
+	if (!AcquiredImageQueue.Dequeue(NextImageToWait))
+	{
+		return XrResult::XR_ERROR_CALL_ORDER_INVALID;
+	}
+
+	FSwapchainImage& NextImage = Images[NextImageToWait];
+	if (NextImage.ImageState == EImageState::Aquired)
+	{
+		NextImage.ImageState = EImageState::Waiting;
+
+		//custom image ready wait could be here, if there was one.
+
+		NextImage.ImageState = EImageState::WaitComplete;
+		return XrResult::XR_SUCCESS;
+	}
+	else
+	{
+		return XrResult::XR_ERROR_CALL_ORDER_INVALID;
+	}
+}
+
+XrResult FOXRVisionOSSwapchain::XrReleaseSwapchainImage(
+	const XrSwapchainImageReleaseInfo*			releaseInfo)
+{
+	SCOPED_NAMED_EVENT_TEXT("FOXRVisionOSSwapchain::XrReleaseSwapchainImage", FColor::Turquoise);
+
+	// releaseInfo can be null
+	if (releaseInfo != nullptr && releaseInfo->type != XR_TYPE_SWAPCHAIN_IMAGE_RELEASE_INFO)
+	{
+		return XrResult::XR_ERROR_VALIDATION_FAILURE;
+	}
+
+	// OXR spec says only one image can be waited on, therefore,
+	// the last used wait index is the outstanding wait
+	NextImageToRelease = NextImageToWait;
+
+	FSwapchainImage& NextImage = Images[NextImageToRelease];
+	if (NextImage.ImageState == EImageState::WaitComplete)
+	{
+				NextImage.ImageState = EImageState::Released;
+		OutstandingWaits--;
+		return XrResult::XR_SUCCESS;
+	}
+	else
+	{
+		return XrResult::XR_ERROR_CALL_ORDER_INVALID;
+	}
+}
diff --git a/Engine/Platforms/VisionOS/Plugins/Runtime/OpenXRVisionOS/Source/OXRVisionOS/Private/OXRVisionOSSwapchain.h b/Engine/Platforms/VisionOS/Plugins/Runtime/OpenXRVisionOS/Source/OXRVisionOS/Private/OXRVisionOSSwapchain.h
new file mode 100644
--- /dev/null
+++ b/Engine/Platforms/VisionOS/Plugins/Runtime/OpenXRVisionOS/Source/OXRVisionOS/Private/OXRVisionOSSwapchain.h
@@ -0,0 +1,77 @@
+// Copyright Epic Games, Inc. All Rights Reserved.
+
+#pragma once
+
+#include "OXRVisionOS.h"
+
+#include "Runtime/Launch/Resources/Version.h"
+#include "Templates/RefCounting.h"
+
+class FOXRVisionOSSession;
+namespace MetalRHIOXRVisionOS
+{
+	class ISwapchainLabelManager;
+}
+
+
+class FOXRVisionOSSwapchain
+{
+public:
+	enum class EImageState
+	{
+		Released,
+		Aquired,
+		Waiting,
+		WaitComplete
+	};
+
+	struct FSwapchainImage
+	{
+		FSwapchainImage(FTextureRHIRef InImage, EImageState InImageState);
+		~FSwapchainImage();
+
+		FTextureRHIRef Image;
+		EImageState ImageState = EImageState::Released;
+	};
+
+	static XrResult Create(TSharedPtr<FOXRVisionOSSwapchain, ESPMode::ThreadSafe>& OutSwapchain, const XrSwapchainCreateInfo* createInfo, uint32 SwapchainLength, FOXRVisionOSSession* Session);
+
+	FOXRVisionOSSwapchain(const XrSwapchainCreateInfo* createInfo, uint32 SwapchainLength, FOXRVisionOSSession* Session);
+	~FOXRVisionOSSwapchain();
+	XrResult XrDestroySwapchain();
+
+	XrResult XrEnumerateSwapchainImages(
+		uint32_t                                    imageCapacityInput,
+		uint32_t*									imageCountOutput,
+		XrSwapchainImageBaseHeader*					images);
+
+	XrResult XrAcquireSwapchainImage(
+		const XrSwapchainImageAcquireInfo*			acquireInfo,
+		uint32_t*									index);
+
+	XrResult XrWaitSwapchainImage(
+		const XrSwapchainImageWaitInfo*				waitInfo);
+
+	XrResult XrReleaseSwapchainImage(
+		const XrSwapchainImageReleaseInfo*			releaseInfo);
+
+	const FSwapchainImage& GetLastReleasedImage() const;
+
+private:
+	//bool WaitLabel_RHIThread(uint32 AcquireImageIndex);
+	void ReleaseLabel_RHIThread(uint32 ReleaseImageIndex, uint64 LabelValue);
+	
+	XrResult CreateFailureResult = XrResult::XR_SUCCESS;
+	FOXRVisionOSSession*		Session = nullptr;
+
+	XrSwapchainCreateInfo CreateInfo;
+
+	TArray<FSwapchainImage> Images;
+	uint32_t NextImageToAcquire = 0;
+	uint32_t NextImageToWait = 0;
+	uint32_t NextImageToRelease = 0;
+
+
+	TQueue<uint32_t> AcquiredImageQueue;
+	uint32_t OutstandingWaits = 0;
+};
diff --git a/Engine/Platforms/VisionOS/Plugins/Runtime/OpenXRVisionOS/Source/OXRVisionOS/Private/OXRVisionOSTracker.cpp b/Engine/Platforms/VisionOS/Plugins/Runtime/OpenXRVisionOS/Source/OXRVisionOS/Private/OXRVisionOSTracker.cpp
new file mode 100644
--- /dev/null
+++ b/Engine/Platforms/VisionOS/Plugins/Runtime/OpenXRVisionOS/Source/OXRVisionOS/Private/OXRVisionOSTracker.cpp
@@ -0,0 +1,350 @@
+// Copyright Epic Games, Inc. All Rights Reserved.
+
+#include "OXRVisionOSTracker.h"
+
+DECLARE_LOG_CATEGORY_CLASS(LogOXRVisionOSTracker, Log, All);
+
+/** includes */
+#include "OXRVisionOSPlatformUtils.h"
+#include "OXRVisionOSController.h"
+#include "OXRVisionOSInstance.h"
+#include "OXRVisionOSSession.h"
+
+bool FOXRVisionOSTracker::Create(TSharedPtr<FOXRVisionOSTracker, ESPMode::ThreadSafe>& OutTracker, bool bEnableGaze)
+{
+	OutTracker = MakeShared<FOXRVisionOSTracker, ESPMode::ThreadSafe>(bEnableGaze);
+	if (OutTracker->bCreateFailed)
+	{
+		OutTracker = nullptr;
+		return false;
+	}
+	return true;
+}
+
+//=============================================================================
+FOXRVisionOSTracker::FOXRVisionOSTracker(bool bEnableGaze)// :
+//	bGazeEnabled(bEnableGaze)
+{
+
+	//if (Ret != OK)
+	//{
+	//	UE_LOG(LogOXRVisionOSTracker, Warning, TEXT("Initialize failed: 0x%08X"), Ret);
+	//	bCreateFailed = true;
+	//	return;
+	//}
+}
+
+//=============================================================================
+FOXRVisionOSTracker::~FOXRVisionOSTracker()
+{
+	// if (Result != OK)
+	// {
+	// 	UE_LOG(LogOXRVisionOSTracker, Warning, TEXT("Term failed: 0x%08X"), Result);
+	// }
+}
+
+
+bool FOXRVisionOSTracker::RegisterDevice(TrackerDeviceType DeviceType, int32 DeviceHandle)
+{
+	FScopeLock ScopeLock(&DeviceInfoMapCriticalSection);
+
+	FDeviceInfo& Added = DeviceInfoMap.Add(DeviceType, FDeviceInfo(DeviceHandle));
+
+	TryRegisterDevice(DeviceType, Added);
+
+	return true;
+}
+
+bool FOXRVisionOSTracker::UnregisterDevice(int32 DeviceHandle)
+{
+	FScopeLock ScopeLock(&DeviceInfoMapCriticalSection);
+
+	for (auto& Elem : DeviceInfoMap)
+	{
+		if (DeviceHandle == Elem.Value.Handle)
+		{
+			DeviceInfoMap.Remove(Elem.Key);
+			break;
+		}
+	}
+	return true;
+}
+
+#define OXRVISIONOS_TRACKER_RESULT_LOGGING 0
+
+void FOXRVisionOSTracker::SyncActions(int32 FrameCounter, XrTime PredictedDisplayTime)
+{
+	FScopeLock ScopeLock(&DeviceInfoMapCriticalSection);
+
+	// // Gaze tracker result caching
+	// if (bGazeEnabled)
+	// {
+	// 	// Null device info means we have not registered the hmd for tracking yet, so we can't get a result
+	// 	if (const FDeviceInfo* DeviceInfo = DeviceInfoMap.Find(DEVICE_HMD))
+	// 	{
+	// 		GazeResults.LastResultIndex = (GazeResults.LastResultIndex + 1) % FGazeResults::MaxResultBufferSize;
+	// 		GazeResults.NumResultsBuffered = FMath::Max(GazeResults.NumResultsBuffered, GazeResults.LastResultIndex + 1);
+	// 		GazeResult& Result = GazeResults.ResultBuffer[GazeResults.LastResultIndex];
+
+	// 		GazeGetResult(&Result);
+	// 		if (Result.Failed)
+	// 		{
+	// 			UE_LOG(LogOXRVisionOSTracker, Warning, TEXT("SyncActions GazeGetResult failed"));
+	// 		}
+	// 	}
+	// }
+}
+
+void FOXRVisionOSTracker::TryRegisterDevice(TrackerDeviceType DeviceType, FDeviceInfo& Info)
+{
+	// int32 Result = APIRegisterDevice(DeviceType, Info.Handle);
+	// if (Result != RESULT_OK)
+	// {
+	// 	UE_LOG(LogOXRVisionOSTracker, Warning, TEXT("RegisterDevice failed code 0x%x. \n"), Result);
+	// }
+	// else
+	// {
+	// 	Info.bIsActive = true;
+	// }
+
+	if (DeviceType == TrackerDeviceType::DEVICE_VR_HMD)
+	{
+		//TODO
+		Info.Handle = 0;
+		Info.bIsActive = true;
+	}
+	
+}
+
+bool FOXRVisionOSTracker::GetResult(int32 DeviceHandle, XrTime DisplayTime, OXRVisionOS::TrackerResultData& ResultData)
+{
+	SCOPED_NAMED_EVENT_TEXT("FOXRVisionOSTracker::GetResult", FColor::Turquoise);
+
+	FScopeLock ScopeLock(&DeviceInfoMapCriticalSection);
+
+	bool bFound = false;
+	TrackerDeviceType DeviceType;
+	FDeviceInfo* DeviceInfo = nullptr;
+	for (auto& Elem : DeviceInfoMap)
+	{
+		if (DeviceHandle == Elem.Value.Handle)
+		{
+			if (Elem.Value.bIsActive == false)
+			{
+				// Device not yet active.
+				return false;
+			}
+			else
+			{
+				DeviceInfo = &Elem.Value;
+				DeviceType = Elem.Key;
+				bFound = true;
+				break;
+			}
+		}
+	}
+
+	if (bFound == false)
+	{
+		UE_LOG(LogOXRVisionOSTracker, Warning, TEXT("FOXRVisionOSTracker::GetResult called with unrecognized DeviceHandle %i"), DeviceHandle);
+		return false;
+	}
+
+	//TODO: get the pose data
+	//{
+	//	SCOPED_NAMED_EVENT_TEXT("GetResult", FColor::Turquoise);
+	//	ResultData.Pose =
+	//}
+
+// #if OXRVISIONOS_TRACKER_RESULT_LOGGING
+// 	LogGetResult(DeviceType, Result, ResultData);
+// #endif
+
+// 	if (Result != RESULT_OK)
+// 	{
+// 		UE_LOG(LogOXRVisionOSTracker, Warning, TEXT("GetResult failed: %d (%u, %08X)"), Result, Result, Result);
+// 		return false;
+// 	}
+
+	return true;
+}
+
+// bool FOXRVisionOSTracker::GetGazeResult(int32 DeviceHandle, XrTime DisplayTime, int32 FrameNumber, GazeResult& OutGazeResult, ResultData& OutHmdResult)
+// {
+// 	if (bGazeEnabled == false)
+// 	{
+// 		UE_LOG(LogOXRVisionOSTracker, Warning, TEXT("FOXRVisionOSTracker::GetGazeResult called without gaze tracking enabled"));
+// 		return false;
+// 	}
+
+// 	OutGazeResult.GazeDirection.x = 0.0f;
+// 	OutGazeResult.GazeDirection.y = 0.0f;
+// 	OutGazeResult.GazeDirection.z = 1.0f;
+
+// 	// Search GazeResults for the result closest to DisplayTime.
+// 	// https://www.khronos.org/registry/OpenXR/specs/1.0/html/xrspec.html#_eye_gaze_input
+// 	// One particularity for eye trackers compared to most other spatial input is that the runtime may not have the capability to predict or interpolate eye gaze poses.
+// 	// Runtimes that cannot predict or interpolate eye gaze poses must clamp the gaze pose requested in the xrLocateSpace call to the value nearest to time requested in the call.
+// 	// We will search starting with the most recent.
+// 	GazeResult* FoundGazeResult = nullptr;
+// 	uint64_t TargetTime = OXRVisionOS::XrTimeToApiTime(DisplayTime);
+// 	uint64_t TimeDiff = UINT64_MAX;
+// 	if (GazeResults.NumResultsBuffered == 0)
+// 	{
+// 		return false; // no results yet
+// 	}
+// 	const int NumResults = GazeResults.NumResultsBuffered;
+// 	for (int i = 0; i < NumResults; ++i)
+// 	{
+// 		int Index = (GazeResults.LastResultIndex - i + NumResults) % NumResults;
+// 		GazeResult& Result = GazeResults.ResultBuffer[Index];
+
+// 		uint64_t NewTimeDiff;
+// 		if (Result.timestamp == TargetTime)  // this may seem unlikely, but the timestamps may come from a common source.
+// 		{
+// 			FoundGazeResult = &Result;
+// 			break;
+// 		}
+// 		else if (TargetTime > Result.timestamp) // result older
+// 		{
+// 			NewTimeDiff = TargetTime - Result.timestamp;
+// 		}
+// 		else // result newer
+// 		{
+// 			NewTimeDiff = Result.timestamp - TargetTime;
+// 		}
+
+// 		if (NewTimeDiff < TimeDiff)
+// 		{
+// 			// This one is better.
+// 			FoundGazeResult = &Result;
+// 			TimeDiff = NewTimeDiff;
+// 		}
+// 		else
+// 		{
+// 			// This one is worse, so the last one was best.
+// 			break;
+// 		}
+// 	}
+
+// 	check(FoundGazeResult);
+// 	OutGazeResult = *FoundGazeResult;
+
+// 	bool bIsGazeDataAvailable = false;
+// 	if (OutGazeResult.status == ENABLED)
+// 	{
+// 		if (OutGazeResult.isGazeValid)
+// 		{
+// 			// gaze data is up to date
+// 			bIsGazeDataAvailable = true;
+// 		}
+// 	}
+// 	else
+// 	{
+// 		bIsGazeDataAvailable = false;
+// 	}
+
+// 	if (!bIsGazeDataAvailable)
+// 	{
+// 		return false;
+// 	}
+
+// 	// Potentially transform gaze data into tracking space, if it isn't already.
+
+// 	if (Result != RESULT_OK)
+// 	{
+// 		UE_LOG(LogOXRVisionOSTracker, Warning, TEXT("GetResult failed: %d (%u, %08X)"), Result, Result, Result);
+// 		return false;
+// 	}
+
+// 	return true;
+// }
+
+bool FOXRVisionOSTracker::UpdateReferenceSpaces(class FOXRVisionOSInstance* Instance, FOXRVisionOSSession* Session, int32 DeviceHandle, XrTime DisplayTime, int32 FrameNumber)
+{
+	OXRVisionOS::TrackerResultData ResultData = {};
+	bool bSuccess = GetResult(DeviceHandle, DisplayTime, ResultData);
+	if (!bSuccess)
+	{
+		return false;
+	}
+
+	// if (ResultData.playAreaUpdateCounter == PlayAreaData.PlayAreaUpdateCounter)
+	// {
+	// 	return false;
+	// }
+
+
+	// UE_LOG(LogOXRVisionOSTracker, Log, TEXT("FOXRVisionOSTracker::UpdateReferenceSpaces update happening. %i"), ResultData.playAreaUpdateCounter);
+	// PlayAreaData.PlayAreaUpdateCounter = ResultData.playAreaUpdateCounter;
+
+	// if (PlayAreaData.PlayAreaMode != ResultData.playAreaMode)
+	// {
+	// 	UE_LOG(LogOXRVisionOSTracker, Log, TEXT("FOXRVisionOSTracker::UpdateReferenceSpaces playAreaMode changing from %i to %i"), PlayAreaData.PlayAreaMode, ResultData.playAreaMode);
+	// 	PlayAreaData.PlayAreaMode = ResultData.playAreaMode;
+	// }
+
+//	{
+//		check(Instance);
+//		check(Session);
+//
+//		// Send the bounds updated event!
+//		XrEventDataReferenceSpaceChangePending Event;
+//		Event.type = XR_TYPE_EVENT_DATA_REFERENCE_SPACE_CHANGE_PENDING;
+//		Event.next = nullptr;
+//		Event.session = (XrSession)this;
+//		Event.referenceSpaceType = XR_REFERENCE_SPACE_TYPE_STAGE;
+//		Event.changeTime = Session->GetCurrentTime();
+//		Event.poseValid = false;//PlayAreaData.PlayAreaOrientedBoundingBoxDataIsValid;  // Currently we are not trying to provide the pose change update here, the spec allows this.  See XrEventDataReferenceSpaceChangePending in the OpenXR reference pages.
+//		//Event.poseInPreviousSpace = ;
+//
+//		Instance->EnqueueEvent(Event);
+//	}
+
+	return true;
+}
+
+bool FOXRVisionOSTracker::GetStagePose(OXRVisionOS::TrackerPoseData& OutPoseData) const
+{
+	// OpenXR requires us to enumerate spaces on startup.
+	// OpenXRHMD assumes that if it can enumerate a stage space that will work.  I'm not sure OpenXR actually guarantees that.
+	// It talks about a stage space unable to locate...
+	// int32_t Result = LocateCoordinate(STAGE, &OutPoseData);
+	// return (Result == RESULT_OK);
+	return true;
+}
+
+void FOXRVisionOSTracker::GetActionStatePose(EOXRVisionOSControllerButton Button, bool& OutActive)
+{
+	// Merge rule: doesn't change
+	// We are only returning active here, one gets the pose via xrCreateActionSpace and xrLocateSpace.
+
+	FScopeLock ScopeLock(&DeviceInfoMapCriticalSection);
+
+	switch (Button)
+	{
+	case EOXRVisionOSControllerButton::GripL:
+	case EOXRVisionOSControllerButton::AimL:
+	{
+		FDeviceInfo* DeviceInfo = DeviceInfoMap.Find(TrackerDeviceType::DEVICE_VR_CONTROLLER_LEFT);
+		OutActive = DeviceInfo->bIsActive; //TODO is this correct?
+		return; 
+	}
+	case EOXRVisionOSControllerButton::GripR:
+	case EOXRVisionOSControllerButton::AimR:
+	{
+		FDeviceInfo* DeviceInfo = DeviceInfoMap.Find(TrackerDeviceType::DEVICE_VR_CONTROLLER_RIGHT);
+		OutActive = DeviceInfo->bIsActive; //TODO is this correct?
+		return;
+	}
+	// case EOXRVisionOSControllerButton::GazePose:
+	// {
+	// 	OutActive = GazeResults.NumResultsBuffered > 0 && (GazeResults.ResultBuffer[GazeResults.LastResultIndex].status == GAZE_TRACKING_STATUS_ENABLE);
+	// 	return;
+	// }
+	default:
+		check(false);
+		OutActive = false;
+		return;
+	}
+}
diff --git a/Engine/Platforms/VisionOS/Plugins/Runtime/OpenXRVisionOS/Source/OXRVisionOS/Private/OXRVisionOSTracker.h b/Engine/Platforms/VisionOS/Plugins/Runtime/OpenXRVisionOS/Source/OXRVisionOS/Private/OXRVisionOSTracker.h
new file mode 100644
--- /dev/null
+++ b/Engine/Platforms/VisionOS/Plugins/Runtime/OpenXRVisionOS/Source/OXRVisionOS/Private/OXRVisionOSTracker.h
@@ -0,0 +1,75 @@
+// Copyright Epic Games, Inc. All Rights Reserved.
+
+#pragma once
+
+#include "CoreMinimal.h"
+#include "OpenXRCore.h"
+#include "Misc/ScopeLock.h"
+
+enum class EOXRVisionOSControllerButton : int32;
+class FOXRVisionOSInstance;
+class FOXRVisionOSSession;
+
+namespace OXRVisionOS
+{
+	struct TrackerResultData;
+	struct TrackerPoseData;
+}
+
+enum class TrackerDeviceType : int32
+{
+	DEVICE_VR_HMD,
+	DEVICE_VR_CONTROLLER_LEFT,
+	DEVICE_VR_CONTROLLER_RIGHT
+};
+
+class FOXRVisionOSTracker
+{
+public:
+	static bool Create(TSharedPtr<FOXRVisionOSTracker, ESPMode::ThreadSafe>& OutTracker, bool bEnableGaze);
+	FOXRVisionOSTracker(bool bEnableGaze);
+	~FOXRVisionOSTracker();
+
+	bool RegisterDevice(TrackerDeviceType DeviceType, int32 DeviceHandle);
+	bool UnregisterDevice(int32 DeviceHandle);
+
+	bool GetResult(int32 DeviceHandle, XrTime DisplayTime, OXRVisionOS::TrackerResultData& OutResultData);
+//	bool GetGazeResult(int32 DeviceHandle, XrTime DisplayTime, GazeResult& OutResult, ResultData& OutHmdResult);
+	bool UpdateReferenceSpaces(FOXRVisionOSInstance* Instance, FOXRVisionOSSession* Session, int32 DeviceHandle, XrTime DisplayTime, int32 FrameNumber);
+	bool GetStagePose(OXRVisionOS::TrackerPoseData& OutPoseData) const; // Get the pose of the stage space in device tracking space.
+
+	void GetActionStatePose(EOXRVisionOSControllerButton Button, bool& OutActive);
+
+	void SyncActions(int32 FrameCounter, XrTime PredictedDisplayTime);
+
+private:
+	const int32 RESULT_OK = 0; //TEMP
+
+	bool bCreateFailed = false;
+	void* TrackingRam = nullptr;
+
+	struct FDeviceInfo
+	{
+		explicit FDeviceInfo(int32 InHandle)
+			: Handle(InHandle)
+		{}
+
+		int32 Handle = -1;
+		bool bIsActive = false;
+	};
+
+	// bool bGazeEnabled;
+	// // We want to be able to get historical results 50ms in the past.  7 120fps frames is 56ms.
+	// struct FGazeResults
+	// {
+	// 	static const int32 MaxResultBufferSize = 7;
+	// 	int32 LastResultIndex = -1;
+	// 	int32 NumResultsBuffered = 0;
+	// 	dddGazeResult ResultBuffer[MaxResultBufferSize] = {};
+	// } GazeResults;
+
+
+	TMap< TrackerDeviceType, FDeviceInfo> DeviceInfoMap;
+	mutable FCriticalSection DeviceInfoMapCriticalSection;
+	void TryRegisterDevice(TrackerDeviceType DeviceType, FDeviceInfo& Info);
+};
diff --git a/Engine/Platforms/VisionOS/Plugins/Runtime/OpenXRVisionOS/Source/OXRVisionOS/Private/OXRVisionOS_RenderBridge.cpp b/Engine/Platforms/VisionOS/Plugins/Runtime/OpenXRVisionOS/Source/OXRVisionOS/Private/OXRVisionOS_RenderBridge.cpp
new file mode 100644
--- /dev/null
+++ b/Engine/Platforms/VisionOS/Plugins/Runtime/OpenXRVisionOS/Source/OXRVisionOS/Private/OXRVisionOS_RenderBridge.cpp
@@ -0,0 +1,50 @@
+// Copyright Epic Games, Inc. All Rights Reserved.
+
+#include "OXRVisionOS_RenderBridge.h"
+
+#include "OXRVisionOS_Swapchain.h"
+#include "OpenXRCore.h"
+#include "OXRVisionOS_openxr_platform.h"
+#include "OXRVisionOSPlatformRHI.h"
+
+#include "OXRVisionOSInstance.h"
+#include "OXRVisionOSSession.h"
+#include "RHICommandList.h"
+
+class FOXRVisionOSRenderBridge : public FOpenXRRenderBridge
+{
+public:
+	FOXRVisionOSRenderBridge(XrInstance InInstance)
+		: FOpenXRRenderBridge(InInstance)
+		, Binding()
+	{
+	}
+
+	virtual void* GetGraphicsBinding(XrSystemId InSystem) override
+	{
+		Binding.type = (XrStructureType)XR_TYPE_GRAPHICS_BINDING_OXRVISIONOS_EPIC;
+		Binding.next = nullptr;
+		return &Binding;
+	}
+
+	virtual FXRSwapChainPtr CreateSwapchain(XrSession InSession, uint8 Format, uint8& OutFormat, uint32 SizeX, uint32 SizeY, uint32 ArraySize, uint32 NumMips, uint32 NumSamples, ETextureCreateFlags CreateFlags, const FClearValueBinding& ClearValueBinding, ETextureCreateFlags AuxiliaryCreateFlags) override final
+	{
+		return CreateSwapchain_OXRVisionOS(InSession, Format, OutFormat, SizeX, SizeY, ArraySize, NumMips, NumSamples, CreateFlags, ClearValueBinding);
+	}
+
+	//virtual bool Support10BitSwapchain() const override { return true; }
+
+	virtual bool HDRGetMetaDataForStereo(EDisplayOutputFormat& OutDisplayOutputFormat, EDisplayColorGamut& OutDisplayColorGamut, bool& OutbHDRSupported) override
+	{
+		OutDisplayOutputFormat = EDisplayOutputFormat::SDR_ExplicitGammaMapping;
+		OutDisplayColorGamut = EDisplayColorGamut::DCIP3_D65;
+		OutbHDRSupported = true;
+		return true;
+	}
+
+private:
+
+	XrGraphicsBindingOXRVisionOSEPIC Binding;
+};
+
+FOpenXRRenderBridge* CreateRenderBridge_OXRVisionOS(XrInstance InInstance) { return new FOXRVisionOSRenderBridge(InInstance); }
diff --git a/Engine/Platforms/VisionOS/Plugins/Runtime/OpenXRVisionOS/Source/OXRVisionOS/Private/OXRVisionOS_RenderBridge.h b/Engine/Platforms/VisionOS/Plugins/Runtime/OpenXRVisionOS/Source/OXRVisionOS/Private/OXRVisionOS_RenderBridge.h
new file mode 100644
--- /dev/null
+++ b/Engine/Platforms/VisionOS/Plugins/Runtime/OpenXRVisionOS/Source/OXRVisionOS/Private/OXRVisionOS_RenderBridge.h
@@ -0,0 +1,8 @@
+// Copyright Epic Games, Inc. All Rights Reserved.
+
+#pragma once
+
+#include "OpenXRHMD_RenderBridge.h"
+
+FOpenXRRenderBridge* CreateRenderBridge_OXRVisionOS(XrInstance InInstance);
+
diff --git a/Engine/Platforms/VisionOS/Plugins/Runtime/OpenXRVisionOS/Source/OXRVisionOS/Private/OXRVisionOS_Swapchain.cpp b/Engine/Platforms/VisionOS/Plugins/Runtime/OpenXRVisionOS/Source/OXRVisionOS/Private/OXRVisionOS_Swapchain.cpp
new file mode 100644
--- /dev/null
+++ b/Engine/Platforms/VisionOS/Plugins/Runtime/OpenXRVisionOS/Source/OXRVisionOS/Private/OXRVisionOS_Swapchain.cpp
@@ -0,0 +1,55 @@
+// Copyright Epic Games, Inc. All Rights Reserved.
+
+#include "OXRVisionOS_Swapchain.h"
+#include "OpenXRCore.h"
+#include "OpenXRPlatformRHI.h"
+#include "XRThreadUtils.h"
+#include "OpenXRHMD_Swapchain.h"
+#include "OXRVisionOS_openxr_platform.h"
+#include "OXRVisionOSPlatformRHI.h"
+
+//UE_TRACE_CHANNEL_EXTERN(OpenXRChannel)
+
+static TArray<XrSwapchainImageOXRVisionOS> EnumerateImagesOXRVisionOS(XrSwapchain InSwapchain, XrStructureType InType)
+{
+	TArray<XrSwapchainImageOXRVisionOS> Images;
+	uint32_t ChainCount;
+	xrEnumerateSwapchainImages(InSwapchain, 0, &ChainCount, nullptr);
+	Images.AddZeroed(ChainCount);
+	for (auto& Image : Images)
+	{
+		Image.type = InType;
+	}
+	XR_ENSURE(xrEnumerateSwapchainImages(InSwapchain, ChainCount, &ChainCount, reinterpret_cast<XrSwapchainImageBaseHeader*>(Images.GetData())));
+	return Images;
+}
+
+FXRSwapChainPtr CreateSwapchain_OXRVisionOS(XrSession InSession, uint8 Format, uint8& OutActualFormat, uint32 SizeX, uint32 SizeY, uint32 ArraySize, uint32 NumMips, uint32 NumSamples, ETextureCreateFlags CreateFlags, const FClearValueBinding& ClearValueBinding)
+{
+	// We use UE formats in OXRVisionOS
+	TFunction<uint32(uint8)> ToPlatformFormat = [](uint8 InFormat) { return InFormat; };
+
+	Format = FOpenXRSwapchain::GetNearestSupportedSwapchainFormat(InSession, Format, ToPlatformFormat);
+	if (!Format)
+	{
+		return nullptr;
+	}
+
+	OutActualFormat = Format;
+	XrSwapchain Swapchain = FOpenXRSwapchain::CreateSwapchain(InSession, Format, SizeX, SizeY, ArraySize, NumMips, NumSamples, CreateFlags);
+	if (!Swapchain)
+	{
+		return nullptr;
+	}
+
+	TArray<FTextureRHIRef> TextureChain;
+	TArray<XrSwapchainImageOXRVisionOS> Images = EnumerateImagesOXRVisionOS(Swapchain, (XrStructureType)XR_TYPE_SWAPCHAIN_IMAGE_OXRVISIONOS_EPIC);
+	for (const auto& Image : Images)
+	{
+		TextureChain.Add(Image.image);
+	}
+	//FTextureRHIRef ChainTarget;//= static_cast<FTextureRHIRef>(GDynamicRHI->RHICreateAliasedTexture((FTextureRHIRef&)TextureChain[0]));
+	FTextureRHIRef ChainTarget = (FTextureRHIRef&)TextureChain[0];
+
+	return CreateXRSwapChain<FOpenXRSwapchain>(MoveTemp(TextureChain), ChainTarget, Swapchain);
+}
diff --git a/Engine/Platforms/VisionOS/Plugins/Runtime/OpenXRVisionOS/Source/OXRVisionOS/Private/OXRVisionOS_Swapchain.h b/Engine/Platforms/VisionOS/Plugins/Runtime/OpenXRVisionOS/Source/OXRVisionOS/Private/OXRVisionOS_Swapchain.h
new file mode 100644
--- /dev/null
+++ b/Engine/Platforms/VisionOS/Plugins/Runtime/OpenXRVisionOS/Source/OXRVisionOS/Private/OXRVisionOS_Swapchain.h
@@ -0,0 +1,12 @@
+// Copyright Epic Games, Inc. All Rights Reserved.
+
+#pragma once
+
+#include "CoreMinimal.h"
+#include "XRSwapChain.h"
+#include "OpenXRPlatformRHI.h"
+#include "Runtime/Launch/Resources/Version.h"
+
+#include <openxr/openxr.h>
+
+FXRSwapChainPtr CreateSwapchain_OXRVisionOS(XrSession InSession, uint8 Format, uint8& OutActualFormat, uint32 SizeX, uint32 SizeY, uint32 ArraySize, uint32 NumMips, uint32 NumSamples, ETextureCreateFlags CreateFlags, const FClearValueBinding& ClearValueBinding);
diff --git a/Engine/Platforms/VisionOS/Plugins/Runtime/OpenXRVisionOS/Source/OXRVisionOSSettings/OXRVisionOSSettings.build.cs b/Engine/Platforms/VisionOS/Plugins/Runtime/OpenXRVisionOS/Source/OXRVisionOSSettings/OXRVisionOSSettings.build.cs
new file mode 100644
--- /dev/null
+++ b/Engine/Platforms/VisionOS/Plugins/Runtime/OpenXRVisionOS/Source/OXRVisionOSSettings/OXRVisionOSSettings.build.cs
@@ -0,0 +1,23 @@
+// Copyright Epic Games, Inc. All Rights Reserved.
+
+using UnrealBuildTool;
+
+[SupportedPlatforms("VisionOS", "Win64", "Mac")]
+public class OXRVisionOSSettings : ModuleRules
+{
+	public OXRVisionOSSettings(ReadOnlyTargetRules Target) : base(Target)
+	{
+		PrivateDependencyModuleNames.AddRange(
+			new string[] {
+                "Core",
+                "CoreUObject",
+                "Engine",
+			}
+		);
+
+		PublicDependencyModuleNames.Add("DeveloperSettings");
+
+		// Needed to include ISettingsModule.h
+		PrivateIncludePathModuleNames.Add("Settings");
+	}
+}
diff --git a/Engine/Platforms/VisionOS/Plugins/Runtime/OpenXRVisionOS/Source/OXRVisionOSSettings/Private/OXRVisionOSSettingsModule.cpp b/Engine/Platforms/VisionOS/Plugins/Runtime/OpenXRVisionOS/Source/OXRVisionOSSettings/Private/OXRVisionOSSettingsModule.cpp
new file mode 100644
--- /dev/null
+++ b/Engine/Platforms/VisionOS/Plugins/Runtime/OpenXRVisionOS/Source/OXRVisionOSSettings/Private/OXRVisionOSSettingsModule.cpp
@@ -0,0 +1,61 @@
+// Copyright Epic Games, Inc. All Rights Reserved.
+
+#include "OXRVisionOSSettingsModule.h"
+#include "OXRVisionOSRuntimeSettings.h"
+#include "Modules/ModuleManager.h"
+#include "UObject/WeakObjectPtr.h"
+#include "UObject/Class.h"
+
+// Settings
+#include "ISettingsModule.h"
+
+#define LOCTEXT_NAMESPACE "OXRVisionOSSettings"
+
+//////////////////////////////////////////////////////////////////////////
+// FOXRVisionOSSettings
+
+class FOXRVisionOSSettings : public IOXRVisionOSSettingsModule
+{
+public:
+	virtual void StartupModule() override
+	{				
+		RegisterSettings();
+	}
+
+	virtual void ShutdownModule() override
+	{
+		if (UObjectInitialized())
+		{
+			UnregisterSettings();  
+		}		
+	}
+private:
+	
+	void RegisterSettings()
+	{
+		if (ISettingsModule* SettingsModule = FModuleManager::GetModulePtr<ISettingsModule>("Settings"))
+		{
+			SettingsModule->RegisterSettings("Project", "Plugins", "OXRVisionOS",
+				LOCTEXT("RuntimeSettingsName", "OXRVisionOS"),
+				LOCTEXT("RuntimeSettingsDescription", "Configure the OXRVisionOS plugin"),
+				GetMutableDefault<UOXRVisionOSRuntimeSettings>()
+			);
+		}
+	}
+
+	void UnregisterSettings()
+	{
+		if (ISettingsModule* SettingsModule = FModuleManager::GetModulePtr<ISettingsModule>("Settings"))
+		{
+			SettingsModule->UnregisterSettings("Project", "Plugins", "OXRVisionOS");
+		}
+	}
+};
+
+//////////////////////////////////////////////////////////////////////////
+
+IMPLEMENT_MODULE(FOXRVisionOSSettings, OXRVisionOSSettings);
+
+//////////////////////////////////////////////////////////////////////////
+
+#undef LOCTEXT_NAMESPACE
diff --git a/Engine/Platforms/VisionOS/Plugins/Runtime/OpenXRVisionOS/Source/OXRVisionOSSettings/Private/OXRVisionOSSettingsModule.h b/Engine/Platforms/VisionOS/Plugins/Runtime/OpenXRVisionOS/Source/OXRVisionOSSettings/Private/OXRVisionOSSettingsModule.h
new file mode 100644
--- /dev/null
+++ b/Engine/Platforms/VisionOS/Plugins/Runtime/OpenXRVisionOS/Source/OXRVisionOSSettings/Private/OXRVisionOSSettingsModule.h
@@ -0,0 +1,16 @@
+// Copyright Epic Games, Inc. All Rights Reserved.
+
+#pragma once
+
+#include "CoreMinimal.h"
+#include "Modules/ModuleInterface.h"
+
+#define OXRVISIONOS_SETTINGS_MODULE_NAME "OXRVisionOSSettings"
+
+//////////////////////////////////////////////////////////////////////////
+// IOXRVisionOSSettingsModule
+
+class IOXRVisionOSSettingsModule : public IModuleInterface
+{
+public:
+};
diff --git a/Engine/Platforms/VisionOS/Plugins/Runtime/OpenXRVisionOS/Source/OXRVisionOSSettings/Public/OXRVisionOSRuntimeSettings.h b/Engine/Platforms/VisionOS/Plugins/Runtime/OpenXRVisionOS/Source/OXRVisionOSSettings/Public/OXRVisionOSRuntimeSettings.h
new file mode 100644
--- /dev/null
+++ b/Engine/Platforms/VisionOS/Plugins/Runtime/OpenXRVisionOS/Source/OXRVisionOSSettings/Public/OXRVisionOSRuntimeSettings.h
@@ -0,0 +1,28 @@
+// Copyright Epic Games, Inc. All Rights Reserved.
+
+#pragma once
+
+#include "CoreMinimal.h"
+#include "Engine/DeveloperSettings.h"
+#include "OXRVisionOSSettingsTypes.h"
+
+#include "OXRVisionOSRuntimeSettings.generated.h"
+
+UCLASS(config=Engine, defaultconfig, meta=(DisplayName="OXRVisionOS"))
+class OXRVISIONOSSETTINGS_API UOXRVisionOSRuntimeSettings : public UDeveloperSettings
+{
+	GENERATED_BODY()
+
+public:
+	// Defines a rotation adjustment to the aim pose from the OXRVisionOS native controller orientation.  This pose may be fine without adjustment.
+	UPROPERTY(config, EditAnywhere, Category=OXRVisionOS, Meta = (DisplayName = "Aim Pose Adjustment"))
+	FRotator OXRVisionOSAimPoseAdjustment = FRotator(0.0f, 0.0f, 0.0f);
+
+	// Defines a rotation adjustment to the grip pose from the OXRVisionOS native controller orientation.  We provide a default adjustment to grip.
+	UPROPERTY(config, EditAnywhere, Category = OXRVisionOS, Meta = (DisplayName = "Grip Pose Adjustment"))
+	FRotator OXRVisionOSGripPoseAdjustment = FRotator(45.0f,0.0f,0.0f);
+
+	// Number of UI2D/quad layers supported by OXRVisionOS, required at app startup time
+	UPROPERTY(config, EditAnywhere, Category = OXRVisionOS, Meta = (DeprecatedProperty, ClampMin = "0", ClampMax = "4", DisplayName = "Max Number UI Layers"))
+	int32 OXRVisionOSNumUILayers = 0;
+};
diff --git a/Engine/Platforms/VisionOS/Plugins/Runtime/OpenXRVisionOS/Source/OXRVisionOSSettings/Public/OXRVisionOSSettingsTypes.h b/Engine/Platforms/VisionOS/Plugins/Runtime/OpenXRVisionOS/Source/OXRVisionOSSettings/Public/OXRVisionOSSettingsTypes.h
new file mode 100644
--- /dev/null
+++ b/Engine/Platforms/VisionOS/Plugins/Runtime/OpenXRVisionOS/Source/OXRVisionOSSettings/Public/OXRVisionOSSettingsTypes.h
@@ -0,0 +1,13 @@
+// Copyright Epic Games, Inc. All Rights Reserved.
+
+#pragma once
+
+#include "CoreMinimal.h"
+
+//#include "OXRVisionOSSettingsTypes.generated.h"
+
+//UENUM()
+//enum class EFoo : int8
+//{
+//	Bar = 0 UMETA(DisplayName = "Baaar"),
+//};
