OXRVisionOS mtlpp removal
-Removed references to mtlpp, which were not being used anyway.

diff --git a/Engine/Platforms/VisionOS/Plugins/Runtime/OpenXRVisionOS/Source/OXRVisionOS/OXRVisionOS.Build.cs b/Engine/Platforms/VisionOS/Plugins/Runtime/OpenXRVisionOS/Source/OXRVisionOS/OXRVisionOS.Build.cs
--- a/Engine/Platforms/VisionOS/Plugins/Runtime/OpenXRVisionOS/Source/OXRVisionOS/OXRVisionOS.Build.cs
+++ b/Engine/Platforms/VisionOS/Plugins/Runtime/OpenXRVisionOS/Source/OXRVisionOS/OXRVisionOS.Build.cs
@@ -22,9 +22,6 @@ namespace UnrealBuildTool.Rules
 				"OXRVisionOSSettings",
 			});
 			
-		AddEngineThirdPartyPrivateStaticDependencies(Target,
-			"MTLPP"
-		);
 			PublicWeakFrameworks.Add("Metal");
 			PublicWeakFrameworks.Add("ARKit");		
 			PrivateIncludePathModuleNames.Add("OpenXRHMD");
diff --git a/Engine/Platforms/VisionOS/Plugins/Runtime/OpenXRVisionOS/Source/OXRVisionOS/Private/OXRVisionOSSwapchain.cpp b/Engine/Platforms/VisionOS/Plugins/Runtime/OpenXRVisionOS/Source/OXRVisionOS/Private/OXRVisionOSSwapchain.cpp
--- a/Engine/Platforms/VisionOS/Plugins/Runtime/OpenXRVisionOS/Source/OXRVisionOS/Private/OXRVisionOSSwapchain.cpp
+++ b/Engine/Platforms/VisionOS/Plugins/Runtime/OpenXRVisionOS/Source/OXRVisionOS/Private/OXRVisionOSSwapchain.cpp
@@ -11,9 +11,6 @@
 @end
 #endif
 #include "HAL/PlatformFramePacer.h"
-THIRD_PARTY_INCLUDES_START
-#include "mtlpp.hpp"
-THIRD_PARTY_INCLUDES_END
 
 #if PLATFORM_VISIONOS
 #import <CompositorServices/CompositorServices.h>
