UE-204751 VisionOS -> visionOS in all user facing strings.
-Fixed plugin project settings UI duplication (settings were displayed twice).

diff --git a/Engine/Platforms/VisionOS/Plugins/Runtime/OpenXRVisionOS/OpenXRVisionOS.uplugin b/Engine/Platforms/VisionOS/Plugins/Runtime/OpenXRVisionOS/OpenXRVisionOS.uplugin
--- a/Engine/Platforms/VisionOS/Plugins/Runtime/OpenXRVisionOS/OpenXRVisionOS.uplugin
+++ b/Engine/Platforms/VisionOS/Plugins/Runtime/OpenXRVisionOS/OpenXRVisionOS.uplugin
@@ -2,8 +2,8 @@
 	"FileVersion": 3,
 	"Version": 1,
 	"VersionName": "1.0",
-	"FriendlyName": "OpenXRVisionOS",
-	"Description": "OpenXRVisionOS exposes VisionOS full immersion to Unreal as if it was an OpenXR runtime.  Documentation on developing VR in Unreal with OpenXR is applicable.",
+	"FriendlyName": "OpenXR visionOS",
+	"Description": "OpenXR visionOS exposes visionOS full immersion to Unreal as if it was an OpenXR runtime.  Documentation on developing VR in Unreal with OpenXR is applicable.",
 	"Category": "Virtual Reality",
 	"CreatedBy": "Epic Games, Inc.",
 	"CreatedByURL": "https://epicgames.com",
diff --git a/Engine/Platforms/VisionOS/Plugins/Runtime/OpenXRVisionOS/Source/OXRVisionOSSettings/Private/OXRVisionOSSettingsModule.cpp b/Engine/Platforms/VisionOS/Plugins/Runtime/OpenXRVisionOS/Source/OXRVisionOSSettings/Private/OXRVisionOSSettingsModule.cpp
--- a/Engine/Platforms/VisionOS/Plugins/Runtime/OpenXRVisionOS/Source/OXRVisionOSSettings/Private/OXRVisionOSSettingsModule.cpp
+++ b/Engine/Platforms/VisionOS/Plugins/Runtime/OpenXRVisionOS/Source/OXRVisionOSSettings/Private/OXRVisionOSSettingsModule.cpp
@@ -16,46 +16,30 @@
 
 class FOXRVisionOSSettings : public IOXRVisionOSSettingsModule
 {
-public:
-	virtual void StartupModule() override
-	{				
-		RegisterSettings();
-	}
-
-	virtual void ShutdownModule() override
-	{
-		if (UObjectInitialized())
-		{
-			UnregisterSettings();  
-		}		
-	}
-private:
-	
-	void RegisterSettings()
-	{
-		if (ISettingsModule* SettingsModule = FModuleManager::GetModulePtr<ISettingsModule>("Settings"))
-		{
-			SettingsModule->RegisterSettings("Project", "Plugins", "OXRVisionOS",
-				LOCTEXT("RuntimeSettingsName", "OXRVisionOS"),
-				LOCTEXT("RuntimeSettingsDescription", "Configure the OXRVisionOS plugin"),
-				GetMutableDefault<UOXRVisionOSRuntimeSettings>()
-			);
-		}
-	}
-
-	void UnregisterSettings()
-	{
-		if (ISettingsModule* SettingsModule = FModuleManager::GetModulePtr<ISettingsModule>("Settings"))
-		{
-			SettingsModule->UnregisterSettings("Project", "Plugins", "OXRVisionOS");
-		}
-	}
 };
 
 //////////////////////////////////////////////////////////////////////////
 
 IMPLEMENT_MODULE(FOXRVisionOSSettings, OXRVisionOSSettings);
 
+//Note: this one is required to get the setting to show up in Editor->Project Settings->Plugins
+FName UOXRVisionOSRuntimeSettings::GetCategoryName() const
+{
+	return TEXT("Plugins");
+}
+
+#if WITH_EDITOR
+FText UOXRVisionOSRuntimeSettings::GetSectionText() const
+{
+	return NSLOCTEXT("OXRVisionOSSettings", "OpenXRVisionOSRuntimeSettingsSection", "OpenXR visionOS");
+}
+
+FText UOXRVisionOSRuntimeSettings::GetSectionDescription() const
+{
+	return NSLOCTEXT("OXRVisionOSSettings", "OpenXRVisionOSRuntimeSettingsSectionDescription", "OpenXR visionOS Settings");
+}
+
+#endif
 //////////////////////////////////////////////////////////////////////////
 
 #undef LOCTEXT_NAMESPACE
diff --git a/Engine/Platforms/VisionOS/Plugins/Runtime/OpenXRVisionOS/Source/OXRVisionOSSettings/Public/OXRVisionOSRuntimeSettings.h b/Engine/Platforms/VisionOS/Plugins/Runtime/OpenXRVisionOS/Source/OXRVisionOSSettings/Public/OXRVisionOSRuntimeSettings.h
--- a/Engine/Platforms/VisionOS/Plugins/Runtime/OpenXRVisionOS/Source/OXRVisionOSSettings/Public/OXRVisionOSRuntimeSettings.h
+++ b/Engine/Platforms/VisionOS/Plugins/Runtime/OpenXRVisionOS/Source/OXRVisionOSSettings/Public/OXRVisionOSRuntimeSettings.h
@@ -8,21 +8,29 @@
 
 #include "OXRVisionOSRuntimeSettings.generated.h"
 
-UCLASS(config=Engine, defaultconfig, meta=(DisplayName="OXRVisionOS"))
+UCLASS(config=Engine, defaultconfig)
 class OXRVISIONOSSETTINGS_API UOXRVisionOSRuntimeSettings : public UDeveloperSettings
 {
 	GENERATED_BODY()
 
 public:
-	// Defines a rotation adjustment to the aim pose from the OXRVisionOS native controller orientation.  This pose may be fine without adjustment.
-	UPROPERTY(config, EditAnywhere, Category=OXRVisionOS, Meta = (DisplayName = "Aim Pose Adjustment"))
+	//~ Begin UDeveloperSettings interface
+	virtual FName GetCategoryName() const;
+#if WITH_EDITOR
+	virtual FText GetSectionText() const override;
+	virtual FText GetSectionDescription() const override;
+#endif
+	//~ End UDeveloperSettings interface
+	
+	// Defines a rotation adjustment to the aim pose from the OpenXR visionOS native controller orientation.  This pose may be fine without adjustment.
+	UPROPERTY(config, EditAnywhere, Category="OpenXR visionOS", Meta = (DisplayName = "Aim Pose Adjustment"))
 	FRotator OXRVisionOSAimPoseAdjustment = FRotator(0.0f, 0.0f, 0.0f);
 
-	// Defines a rotation adjustment to the grip pose from the OXRVisionOS native controller orientation.  We provide a default adjustment to grip.
-	UPROPERTY(config, EditAnywhere, Category = OXRVisionOS, Meta = (DisplayName = "Grip Pose Adjustment"))
+	// Defines a rotation adjustment to the grip pose from the OpenXR visionOS native controller orientation.  We provide a default adjustment to grip.
+	UPROPERTY(config, EditAnywhere, Category = "OpenXR visionOS", Meta = (DisplayName = "Grip Pose Adjustment"))
 	FRotator OXRVisionOSGripPoseAdjustment = FRotator(45.0f,0.0f,0.0f);
 
-	// Number of UI2D/quad layers supported by OXRVisionOS, required at app startup time
-	UPROPERTY(config, EditAnywhere, Category = OXRVisionOS, Meta = (DeprecatedProperty, ClampMin = "0", ClampMax = "4", DisplayName = "Max Number UI Layers"))
+	// Number of UI2D/quad layers supported by OpenXR visionOS, required at app startup time
+	UPROPERTY(config, EditAnywhere, Category = "OpenXR visionOS", Meta = (DeprecatedProperty, ClampMin = "0", ClampMax = "4", DisplayName = "Max Number UI Layers"))
 	int32 OXRVisionOSNumUILayers = 0;
 };
diff --git a/Engine/Platforms/VisionOS/Source/Developer/VisionOSPlatformEditor/Private/VisionOSPlatformEditorModule.cpp b/Engine/Platforms/VisionOS/Source/Developer/VisionOSPlatformEditor/Private/VisionOSPlatformEditorModule.cpp
--- a/Engine/Platforms/VisionOS/Source/Developer/VisionOSPlatformEditor/Private/VisionOSPlatformEditorModule.cpp
+++ b/Engine/Platforms/VisionOS/Source/Developer/VisionOSPlatformEditor/Private/VisionOSPlatformEditorModule.cpp
@@ -23,8 +23,8 @@ class FVisionOSPlatformEditorModule
 		if (SettingsModule != nullptr)
 		{
 			SettingsModule->RegisterSettings("Project", "Platforms", "VisionOS",
-				LOCTEXT("VisionOSSettingsName", "VisionOS"),
-				LOCTEXT("VisionOSSettingsDescription", "Settings for VisionOS projects"),
+				LOCTEXT("VisionOSSettingsName", "visionOS"),
+				LOCTEXT("VisionOSSettingsDescription", "Settings for visionOS projects"),
 				GetMutableDefault<UVisionOSRuntimeSettings>());
 
 		}
diff --git a/Engine/Platforms/VisionOS/Source/Runtime/VisionOSRuntimeSettings/Classes/VisionOSRuntimeSettings.h b/Engine/Platforms/VisionOS/Source/Runtime/VisionOSRuntimeSettings/Classes/VisionOSRuntimeSettings.h
--- a/Engine/Platforms/VisionOS/Source/Runtime/VisionOSRuntimeSettings/Classes/VisionOSRuntimeSettings.h
+++ b/Engine/Platforms/VisionOS/Source/Runtime/VisionOSRuntimeSettings/Classes/VisionOSRuntimeSettings.h
@@ -30,9 +30,9 @@ public:
 	virtual void PostInitProperties() override;
 #endif // WITH_EDITOR
 
-    // When building for VisionOS, default to Full Immersion mode.  If disabled, switches to Window mode. Note: when enabled,
+    // When building for visionOS, default to Full Immersion mode.  If disabled, switches to Window mode. Note: when enabled,
 	// the @main struct in UESwift.swift becomes the new main(). Plus, UIApplicationSceneManifest is auto-generation for Info.plist
-    UPROPERTY(config, EditAnywhere, Category = Build, Meta = (DisplayName = "Enable VisionOS's Immersion Mode (Experimental)", ConfigRestartRequired = true))
+    UPROPERTY(config, EditAnywhere, Category = Build, Meta = (DisplayName = "Enable visionOS's Immersion Mode (Experimental)", ConfigRestartRequired = true))
     bool bUseSwiftUIMain;
 
 	// If checked, the Swift/ObjC bridging headers will be created.
