UE-206181 Experimental support for VisionOS
-First phase of VisionOS running on device.  VRTemplate will run on device.
-Mobile forward renderer runs if MobileMultiView is disabled at reasonable framerate, however depth is currently incorrect (we are just writing scene color to it, depth information exists but is not yet connected).  The blocky black shadow artifact is related to this.
-Deferred rendering runs at approximately 10fps with no anti-aliasing and (if you make a trivial code edit) with a proper depth buffer.
-Appears to be running one frame behind now, not sure how.  This broke when moving up to this stream.
Also review 31739189

diff --git a/Engine/Platforms/VisionOS/Plugins/Runtime/OpenXRVisionOS/Source/OXRVisionOS/Private/OXRVisionOS.cpp b/Engine/Platforms/VisionOS/Plugins/Runtime/OpenXRVisionOS/Source/OXRVisionOS/Private/OXRVisionOS.cpp
--- a/Engine/Platforms/VisionOS/Plugins/Runtime/OpenXRVisionOS/Source/OXRVisionOS/Private/OXRVisionOS.cpp
+++ b/Engine/Platforms/VisionOS/Plugins/Runtime/OpenXRVisionOS/Source/OXRVisionOS/Private/OXRVisionOS.cpp
@@ -779,8 +779,9 @@ void FOXRVisionOS::StartupModule()
 
 	OXRVisionOS::BuildFunctionMaps();
 
+	SupportedExtensions.Add(XR_KHR_COMPOSITION_LAYER_DEPTH_EXTENSION_NAME, XR_KHR_composition_layer_depth_SPEC_VERSION);
 	SupportedExtensions.Add(XR_EPIC_OXRVISIONOS_CONTROLLER_NAME, XR_EPIC_oxrvisionos_controller_SPEC_VERSION);
-//	SupportedExtensions.Add(XR_EXT_EYE_GAZE_INTERACTION_EXTENSION_NAME, XR_EXT_eye_gaze_interaction_SPEC_VERSION);
+	//	SupportedExtensions.Add(XR_EXT_EYE_GAZE_INTERACTION_EXTENSION_NAME, XR_EXT_eye_gaze_interaction_SPEC_VERSION);
 }
 
 void FOXRVisionOS::ShutdownModule()
diff --git a/Engine/Platforms/VisionOS/Plugins/Runtime/OpenXRVisionOS/Source/OXRVisionOS/Private/OXRVisionOS.h b/Engine/Platforms/VisionOS/Plugins/Runtime/OpenXRVisionOS/Source/OXRVisionOS/Private/OXRVisionOS.h
--- a/Engine/Platforms/VisionOS/Plugins/Runtime/OpenXRVisionOS/Source/OXRVisionOS/Private/OXRVisionOS.h
+++ b/Engine/Platforms/VisionOS/Plugins/Runtime/OpenXRVisionOS/Source/OXRVisionOS/Private/OXRVisionOS.h
@@ -15,7 +15,7 @@
 // Uncomment to make intellisense work better.
 //#include "../../../../../../../../Source/ThirdParty/OpenXR/include/openxr/openxr.h"
 
-DECLARE_LOG_CATEGORY_EXTERN(LogOXRVisionOS, Verbose, All);
+DECLARE_LOG_CATEGORY_EXTERN(LogOXRVisionOS, Log, All);
 // Note: if you are turning this log up to verbose you may also want LogMetalVisionOS turned up to verbose.
 
 class FOXRVisionOSInstance;
diff --git a/Engine/Platforms/VisionOS/Plugins/Runtime/OpenXRVisionOS/Source/OXRVisionOS/Private/OXRVisionOSInstance.cpp b/Engine/Platforms/VisionOS/Plugins/Runtime/OpenXRVisionOS/Source/OXRVisionOS/Private/OXRVisionOSInstance.cpp
--- a/Engine/Platforms/VisionOS/Plugins/Runtime/OpenXRVisionOS/Source/OXRVisionOS/Private/OXRVisionOSInstance.cpp
+++ b/Engine/Platforms/VisionOS/Plugins/Runtime/OpenXRVisionOS/Source/OXRVisionOS/Private/OXRVisionOSInstance.cpp
@@ -108,18 +108,9 @@ FOXRVisionOSInstance::FOXRVisionOSInstance(const XrInstanceCreateInfo* CreateInf
     check(NumViewports < 3); // unexpected, and we aren't handling it.
     CGRect LeftRect = OXRVisionOS::GetSwiftViewportRect(0);
     CGRect RightRect = NumViewports > 1 ? OXRVisionOS::GetSwiftViewportRect(1) : CGRectMake(0,0,LeftRect.size.width/4, LeftRect.size.height/4); //TODO CVars around this perhaps???
-    //LeftRect.size.width, LeftRect.size.height
 
-	int32 MaxWidth = LeftRect.size.width + (NumViewports > 1 ? RightRect.size.width : 0);
+	int32 MaxWidth = LeftRect.size.width;
 	int32 MaxHeight = LeftRect.size.height;
-	//if (GConfig->GetInt(TEXT("/Script/OXRVisionOSSettings.OXRVisionOSRuntimeSettings"), TEXT("OXRVisionOSMaxPerEyeWidth"), MaxWidth, GEngineIni))
-	//{
-	//	UE_LOG(LogOXRVisionOS, Log, TEXT("Max texture width read from settings file: %d"), MaxWidth);
-	//}
-	//if (GConfig->GetInt(TEXT("/Script/OXRVisionOSSettings.OXRVisionOSRuntimeSettings"), TEXT("OXRVisionOSMaxPerEyeHeight"), MaxHeight, GEngineIni))
-	//{
-	//	UE_LOG(LogOXRVisionOS, Log, TEXT("Max texture height read from settings file: %d"), MaxHeight);
-	//}
 
 	SystemProperties.graphicsProperties.maxSwapchainImageWidth = MaxWidth;
     SystemProperties.graphicsProperties.maxSwapchainImageHeight = MaxHeight;
diff --git a/Engine/Platforms/VisionOS/Plugins/Runtime/OpenXRVisionOS/Source/OXRVisionOS/Private/OXRVisionOSSession.cpp b/Engine/Platforms/VisionOS/Plugins/Runtime/OpenXRVisionOS/Source/OXRVisionOS/Private/OXRVisionOSSession.cpp
--- a/Engine/Platforms/VisionOS/Plugins/Runtime/OpenXRVisionOS/Source/OXRVisionOS/Private/OXRVisionOSSession.cpp
+++ b/Engine/Platforms/VisionOS/Plugins/Runtime/OpenXRVisionOS/Source/OXRVisionOS/Private/OXRVisionOSSession.cpp
@@ -116,26 +116,39 @@ FOXRVisionOSSession::FOXRVisionOSSession(const XrSessionCreateInfo* createInfo,
 	XrWaitFrameEvent->Trigger();
 
 	SetSessionState(XR_SESSION_STATE_IDLE);
-	
-	if (Instance->IsHMDAvailable())
-	{
-		SetSessionState(XR_SESSION_STATE_READY);
-	}
-	else
-	{
-	//	// else we need to wait for the hmd to become available, not relevant on this all-in-one hmd device
-	//	OnWorldTickStartDelegateHandle = FWorldDelegates::OnWorldTickStart.AddRaw(this, &FOXRVisionOSSession::OnWorldTickStart);
-	}
+
+	// We need to wait for the layerstate to transition before starting.
+	OnWorldTickStartDelegateHandle = FWorldDelegates::OnWorldTickStart.AddRaw(this, &FOXRVisionOSSession::OnWorldTickStart);
 }
 
 void FOXRVisionOSSession::OnWorldTickStart(UWorld* World, ELevelTick TickType, float DeltaTime)
 {
 	if (SessionState == XR_SESSION_STATE_IDLE)
 	{
-		if (Instance->IsHMDAvailable())
+		switch (cp_layer_renderer_get_state(OXRVisionOS::GetSwiftLayerRenderer()))
 		{
-			SetSessionState(XR_SESSION_STATE_READY);
-		}
+			case cp_layer_renderer_state_paused:
+				// Wait until the scene appears.
+				// Polling instead of using this wait function.
+				//cp_layer_renderer_wait_until_running(layer);
+				UE_LOG(LogOXRVisionOS, Warning, TEXT("FOXRVisionOSSession cp_layer_renderer_get_state is paused.  Waiting in XR_SESSION_STATE_IDLE."));
+				break;
+
+			case cp_layer_renderer_state_running:
+				UE_LOG(LogOXRVisionOS, Warning, TEXT("FOXRVisionOSSession cp_layer_renderer_get_state is running. Going to XR_SESSION_STATE_READY"));
+
+				SetSessionState(XR_SESSION_STATE_READY);
+				break;
+
+			case cp_layer_renderer_state_invalidated:
+				// Not entirely sure this is correct, but it's kind of close.
+				UE_LOG(LogOXRVisionOS, Warning, TEXT("FOXRVisionOSSession cp_layer_renderer_get_state is cp_layer_renderer_state_invalidated.  Ending session."));
+				EndSessionInternal();
+				break;
+				
+			default:
+				check(false);
+		 }
 	}
 
 	// Stop trying to transition to READY when no longer IDLE.
@@ -344,16 +357,13 @@ XrResult FOXRVisionOSSession::XrEnumerateSwapchainFormats(
 			return XrResult::XR_ERROR_SIZE_INSUFFICIENT;
 		}
 
-		//TODO is this right?
 		// HDR formats 
-		Formats[0] = PF_B8G8R8A8;			//k8_8_8_8UNorm			// This is Unreal's default format
-		Formats[1] = PF_A2B10G10R10; 		//k10_10_10_2UNorm		
-		Formats[2] = PF_FloatRGB; 			//k11_11_10Float
-		Formats[3] = PF_FloatRGBA; 			//k16_16_16_16Float
+		Formats[0] = PF_B8G8R8A8;			//bgra8Unorm_srgb			// This is Unreal's default format
+		Formats[1] = PF_FloatRGBA; 			//rgba16Float
 
-		//Depth formats - depth is not supported
-		//formats[1] = PF_DepthStencil;		//k32Float
-		//formats[2] = PF_D24;				//k32Float
+		//Depth formats
+		Formats[2] = PF_DepthStencil;		//depth32Float_stencil8		// Correct for deferred
+		Formats[3] = PF_R32_FLOAT;			//depth32Float				// Possibly correct for mobile forward
 	}
 	return XrResult::XR_SUCCESS;
 }
@@ -808,7 +818,28 @@ XrResult FOXRVisionOSSession::XrWaitFrame(
 		UE_LOG(LogOXRVisionOS, Log, TEXT("FOXRVisionOSSession XrWaitFrame XR_ERROR_SESSION_NOT_RUNNING"));
 		return XrResult::XR_ERROR_SESSION_NOT_RUNNING;
 	}
+	
+	XrResult SuccessfulReturnValue =  XrResult::XR_SUCCESS;
+	
+	switch (cp_layer_renderer_get_state(OXRVisionOS::GetSwiftLayerRenderer()))
+	{
+		case cp_layer_renderer_state_paused:
+			//TODO: handle app suspending, this happens when the device locks.  
+			check(false);
+			break;
+
+		case cp_layer_renderer_state_running:
+			break;
 
+		case cp_layer_renderer_state_invalidated:
+			// Not entirely sure this is correct, but it's kind of close.
+			SetSessionState(XrSessionState::XR_SESSION_STATE_LOSS_PENDING);
+			SuccessfulReturnValue = XrResult::XR_SESSION_LOSS_PENDING;
+			break;
+		default:
+			check(false);
+	}
+	
 	if (GetSessionState() == XrSessionState::XR_SESSION_STATE_READY)
 	{
 		SetSessionState(XrSessionState::XR_SESSION_STATE_SYNCHRONIZED);
@@ -822,10 +853,12 @@ XrResult FOXRVisionOSSession::XrWaitFrame(
 	// Block until the previous frame's xrBeginFrame has happened.
 	{
 		SCOPED_NAMED_EVENT_TEXT("FOXRVisionOSSession::XrWaitFrame XrWaitFrameEvent->Wait()", FColor::Turquoise);
+		UE_LOG(LogOXRVisionOS, Verbose, TEXT("XrWaitFrameEvent->Wait() Started"));
 		XrWaitFrameEvent->Wait();
 	}
 
 	PipelinedFrameStateGame.FrameCounter = SessionFrameCounter++;
+	UE_LOG(LogOXRVisionOS, Verbose, TEXT("XrWaitFrameEvent->Wait() Finished 0x%x"), PipelinedFrameStateGame.SwiftFrame);
 
 	//UE_LOG(LogOXRVisionOS, Log, TEXT("FOXRVisionOSSession XrWaitFrame proceeds Game fliparg: %d BBIndex: %d"), PipelinedFrameStateGame.FrameCounter, PipelinedFrameStateGame.BackBufferIndex);
 
@@ -838,8 +871,8 @@ XrResult FOXRVisionOSSession::XrWaitFrame(
 	// }
     
     PipelinedFrameStateGame.SwiftFrame = cp_layer_renderer_query_next_frame(VOSLayerRenderer);
-    UE_LOG(LogTemp, Log, TEXT("SwiftLayerFrame(0x%x) from cp_layer_renderer_query_next_frame in FOXRVisionOSSession::XrWaitFrame"), PipelinedFrameStateGame.SwiftFrame);
-    
+    UE_LOG(LogOXRVisionOS, Verbose, TEXT("SwiftLayerFrame(0x%x) from cp_layer_renderer_query_next_frame in FOXRVisionOSSession::XrWaitFrame"), PipelinedFrameStateGame.SwiftFrame);
+
     // Fetch the predicted timing information.
     PipelinedFrameStateGame.SwiftFrameTiming = cp_frame_predict_timing(PipelinedFrameStateGame.SwiftFrame);
     if ( PipelinedFrameStateGame.SwiftFrameTiming == nullptr)
@@ -852,7 +885,7 @@ XrResult FOXRVisionOSSession::XrWaitFrame(
     }
 
     // Mark the beginnng of our game thread update work.  All the stuff before rendering.
-    UE_LOG(LogTemp, Log, TEXT("SwiftLayerFrame(0x%x) from cp_frame_start_update in FOXRVisionOSSession::XrWaitFrame"), PipelinedFrameStateGame.SwiftFrame);
+    UE_LOG(LogOXRVisionOS, Verbose, TEXT("SwiftLayerFrame(0x%x) from cp_frame_start_update in FOXRVisionOSSession::XrWaitFrame"), PipelinedFrameStateGame.SwiftFrame);
     cp_frame_start_update(PipelinedFrameStateGame.SwiftFrame);
 
     const cp_time_t PresentationTime = cp_frame_timing_get_presentation_time(PipelinedFrameStateGame.SwiftFrameTiming);
@@ -865,38 +898,39 @@ XrResult FOXRVisionOSSession::XrWaitFrame(
     
     PipelinedFrameStateGame.DeviceAnchor = ar_device_anchor_create();
     check(ARKitWorldTrackingProvider)
-    auto anchor_status = ar_world_tracking_provider_query_device_anchor_at_timestamp(ARKitWorldTrackingProvider, TargetRenderTime, PipelinedFrameStateGame.DeviceAnchor);
-    if (anchor_status == ar_device_anchor_query_status_success)
-    {
-        // Get the HMD transform and cache it
-        PipelinedFrameStateGame.HeadTransform = ar_anchor_get_origin_from_anchor_transform(PipelinedFrameStateGame.DeviceAnchor);
-    }
-    else
-    {
-        // Cache identity as the HMD transform as a fallback
-        PipelinedFrameStateGame.HeadTransform = matrix_identity_float4x4;
-    }
+	UE_LOG(LogOXRVisionOS, Verbose, TEXT("SwiftLayerFrame TargetRenderTime= %f"), (double)TargetRenderTime);
+	{
+		//TODO: use a previous frame transform here, because we can't call this function from multiple threads
+		//TEMP use a lock.
+		FScopeLock ScopeLock(&CriticalSection_ar_world_tracking_provider_query_device_anchor_at_timestamp);
+		auto anchor_status = ar_world_tracking_provider_query_device_anchor_at_timestamp(ARKitWorldTrackingProvider, TargetRenderTime, PipelinedFrameStateGame.DeviceAnchor);
+		if (anchor_status == ar_device_anchor_query_status_success)
+		{
+			// Get the HMD transform and cache it
+			PipelinedFrameStateGame.HeadTransform = ar_anchor_get_origin_from_anchor_transform(PipelinedFrameStateGame.DeviceAnchor);
+		}
+		else
+		{
+			// Cache identity as the HMD transform as a fallback
+			PipelinedFrameStateGame.HeadTransform = matrix_identity_float4x4;
+		}
+	}
     
     LocateViewInfoIndexAdvance();
 
-	//TODO XrWaitFrame shouldRender based on visible & focused state
-	frameState->shouldRender = XR_TRUE;
+	frameState->shouldRender = SessionState == XR_SESSION_STATE_FOCUSED || SessionState == XR_SESSION_STATE_VISIBLE;
 
 	// We have to pipeline state updates in order to prevent task overlap (esp during map loading where we kick extra renders)
 	ENQUEUE_RENDER_COMMAND(UpdatePipelinedFrameState)([this, NewRenderingFrameState = PipelinedFrameStateGame](FRHICommandListImmediate& RHICmdList)
 	{
 		check(IsInRenderingThread());
 		PipelinedFrameStateRendering = NewRenderingFrameState;
-		RHICmdList.EnqueueLambda([this, NewRHIFrameState = PipelinedFrameStateRendering](FRHICommandListImmediate&)
-		{
-			PipelinedFrameStateRHI = NewRHIFrameState;
-		});
 	});
 
 	// Update the STAGE space
 	//Tracker->UpdateReferenceSpaces(Instance, this, HMDHandle, PipelinedFrameStateGame.PredictedDisplayTime, PipelinedFrameStateGame.FrameCounter);
 
-	return XrResult::XR_SUCCESS;
+	return SuccessfulReturnValue;
 }
 
 void FOXRVisionOSSession::OnBeginRendering_GameThread()
@@ -904,7 +938,7 @@ void FOXRVisionOSSession::OnBeginRendering_GameThread()
     // Mark the end of our previous game frame
     if (PipelinedFrameStateGame.SwiftFrame)
     {
-        UE_LOG(LogTemp, Log, TEXT("SwiftLayerFrame(0x%x) from cp_frame_end_update in FOXRVisionOSSession::OnBeginRendering_GameThread"), PipelinedFrameStateGame.SwiftFrame);
+        UE_LOG(LogOXRVisionOS, VeryVerbose, TEXT("SwiftLayerFrame(0x%x) from cp_frame_end_update in FOXRVisionOSSession::OnBeginRendering_GameThread"), PipelinedFrameStateGame.SwiftFrame);
         cp_frame_end_update(PipelinedFrameStateGame.SwiftFrame);
         PipelinedFrameStateGame.SwiftFrame = nullptr;
     }
@@ -913,10 +947,14 @@ void FOXRVisionOSSession::OnBeginRendering_GameThread()
 // FOpenXRHMD::OnBeginRendering_RenderThread
 void FOXRVisionOSSession::OnBeginRendering_RenderThread()
 {
-    const MetalRHIVisionOS::BeginRenderingImmersiveParams Params{ PipelinedFrameStateRendering.SwiftFrame };
-    MetalRHIVisionOS::BeginRenderingImmersive(Params);
     WaitUntil();
     StartSubmission();
+	
+	FRHICommandListImmediate& RHICmdList = FRHICommandListImmediate::Get();
+	RHICmdList.EnqueueLambda([this, NewRHIFrameState = PipelinedFrameStateRendering](FRHICommandListImmediate&)
+	{
+		PipelinedFrameStateRHI = NewRHIFrameState;
+	});
 }
 
 // FOpenXRHMD::OnBeginRendering_RHIThread
@@ -976,28 +1014,12 @@ XrResult FOXRVisionOSSession::XrBeginFrame(
     }
     
     {
-        //TODO idk, doing this twice once for render once for rhi...
-        FPipelinedFrameState& FrameState = PipelinedFrameStateRHI;
+		FPipelinedFrameState& FrameState = PipelinedFrameStateRHI;
         
-        UE_LOG(LogTemp, Log, TEXT("SwiftLayerFrame(0x%x) cp_frame_query_drawable in FOXRVisionOSSession::XrBeginFrame"), FrameState.SwiftFrame);
+        UE_LOG(LogOXRVisionOS, VeryVerbose, TEXT("SwiftLayerFrame(0x%x) cp_frame_query_drawable in FOXRVisionOSSession::XrBeginFrame"), FrameState.SwiftFrame);
         FrameState.SwiftDrawable = cp_frame_query_drawable(FrameState.SwiftFrame);
         check(FrameState.SwiftDrawable);
-        UE_LOG(LogTemp, Log, TEXT("SwiftLayerFrame(0x%x) cp_frame_query_drawable got SwiftDrawable(0x%x) in FOXRVisionOSSession::StartSubmission()"), FrameState.SwiftFrame, FrameState.SwiftDrawable);
-
-        FrameState.SwiftFinalFrameTiming = cp_drawable_get_frame_timing(FrameState.SwiftDrawable);
-        
-        CFTimeInterval TimeInterval = cp_time_to_cf_time_interval(cp_frame_timing_get_presentation_time(FrameState.SwiftFinalFrameTiming));
-        check(ARKitWorldTrackingProvider)
-        auto anchor_status = ar_world_tracking_provider_query_device_anchor_at_timestamp(ARKitWorldTrackingProvider, TimeInterval, FrameState.DeviceAnchor);
-        if (anchor_status == ar_device_anchor_query_status_success)
-        {
-            // Get the HMD transform and cache it
-            FrameState.HeadTransform = ar_anchor_get_origin_from_anchor_transform(FrameState.DeviceAnchor);
-        }
-        else
-        {
-            // Leave FrameState.HeadTransform untouched, same as game thread. It may simply be identity.  But whatever it is we will go ahead and render with it.
-        }
+        UE_LOG(LogOXRVisionOS, VeryVerbose, TEXT("SwiftLayerFrame(0x%x) cp_frame_query_drawable got SwiftDrawable(0x%x) in FOXRVisionOSSession::StartSubmission()"), FrameState.SwiftFrame, FrameState.SwiftDrawable);
     }
     
     //MetalRHIOXRVisionOS::WaitUntilSafeForRendering(PipelinedFrameStateRHI.CommandListContext, NextBackBufferIndex);
@@ -1010,9 +1032,9 @@ void FOXRVisionOSSession::WaitUntil()
 {
     FPipelinedFrameState& FrameState = GetPipelinedFrameStateForThread();
     
-    UE_LOG(LogTemp, Log, TEXT("SwiftLayerFrame(0x%x) cp_time_wait_until starting in FOXRVisionOSSession::WaitUntil"), FrameState.SwiftFrame);
+    UE_LOG(LogOXRVisionOS, VeryVerbose, TEXT("SwiftLayerFrame(0x%x) cp_time_wait_until starting in FOXRVisionOSSession::WaitUntil"), FrameState.SwiftFrame);
     cp_time_wait_until(cp_frame_timing_get_optimal_input_time(FrameState.SwiftFrameTiming));
-    UE_LOG(LogTemp, Log, TEXT("SwiftLayerFrame(0x%x) cp_time_wait_until completed in FOXRVisionOSSession::WaitUntil"), FrameState.SwiftFrame);
+    UE_LOG(LogOXRVisionOS, VeryVerbose, TEXT("SwiftLayerFrame(0x%x) cp_time_wait_until completed in FOXRVisionOSSession::WaitUntil"), FrameState.SwiftFrame);
 }
 
 void FOXRVisionOSSession::StartSubmission()
@@ -1020,31 +1042,35 @@ void FOXRVisionOSSession::StartSubmission()
     FPipelinedFrameState& FrameState = GetPipelinedFrameStateForThread();
     
     // tell the compositor service we are going to submit soon and need the drawable
-    UE_LOG(LogTemp, Log, TEXT("SwiftLayerFrame(0x%x) cp_frame_start_submission in FOXRVisionOSSession::XrBeginFrame"), FrameState.SwiftFrame);
+    UE_LOG(LogOXRVisionOS, VeryVerbose, TEXT("SwiftLayerFrame(0x%x) cp_frame_start_submission in FOXRVisionOSSession::XrBeginFrame"), FrameState.SwiftFrame);
     cp_frame_start_submission(FrameState.SwiftFrame);
 
-    UE_LOG(LogTemp, Log, TEXT("SwiftLayerFrame(0x%x) cp_frame_query_drawable in FOXRVisionOSSession::XrBeginFrame"), FrameState.SwiftFrame);
+    UE_LOG(LogOXRVisionOS, VeryVerbose, TEXT("SwiftLayerFrame(0x%x) cp_frame_query_drawable in FOXRVisionOSSession::XrBeginFrame"), FrameState.SwiftFrame);
     FrameState.SwiftDrawable = cp_frame_query_drawable(FrameState.SwiftFrame);
     check(FrameState.SwiftDrawable);
-    UE_LOG(LogTemp, Log, TEXT("SwiftLayerFrame(0x%x) cp_frame_query_drawable got SwiftDrawable(0x%x) in FOXRVisionOSSession::StartSubmission()"), FrameState.SwiftFrame, FrameState.SwiftDrawable);
+    UE_LOG(LogOXRVisionOS, VeryVerbose, TEXT("SwiftLayerFrame(0x%x) cp_frame_query_drawable got SwiftDrawable(0x%x) in FOXRVisionOSSession::StartSubmission()"), FrameState.SwiftFrame, FrameState.SwiftDrawable);
 
     FrameState.SwiftFinalFrameTiming = cp_drawable_get_frame_timing(FrameState.SwiftDrawable);
     
     CFTimeInterval TimeInterval = cp_time_to_cf_time_interval(cp_frame_timing_get_presentation_time(FrameState.SwiftFinalFrameTiming));
     check(ARKitWorldTrackingProvider)
-    auto anchor_status = ar_world_tracking_provider_query_device_anchor_at_timestamp(ARKitWorldTrackingProvider, TimeInterval, FrameState.DeviceAnchor);
-    if (anchor_status == ar_device_anchor_query_status_success)
-    {
-        cp_drawable_set_device_anchor(FrameState.SwiftDrawable, FrameState.DeviceAnchor);
-        // Get the HMD transform and cache it
-        FrameState.HeadTransform = ar_anchor_get_origin_from_anchor_transform(FrameState.DeviceAnchor);
-    }
-    else
-    {
-        //TODO: this is also how you disable reprojection.  Hook that up to a cvar.
-        cp_drawable_set_device_anchor(FrameState.SwiftDrawable, nullptr);
-        // Leave FrameState.HeadTransform untouched, same as game thread. It may simply be identity.  But whatever it is we will go ahead and render with it.
-    }
+	{
+		//TEMP use a lock.
+		FScopeLock ScopeLock(&CriticalSection_ar_world_tracking_provider_query_device_anchor_at_timestamp);
+		auto anchor_status = ar_world_tracking_provider_query_device_anchor_at_timestamp(ARKitWorldTrackingProvider, TimeInterval, FrameState.DeviceAnchor);
+		if (anchor_status == ar_device_anchor_query_status_success)
+		{
+			cp_drawable_set_device_anchor(FrameState.SwiftDrawable, FrameState.DeviceAnchor);
+			// Get the HMD transform and cache it
+			FrameState.HeadTransform = ar_anchor_get_origin_from_anchor_transform(FrameState.DeviceAnchor);
+		}
+		else
+		{
+			//TODO: this is also how you disable reprojection.  Hook that up to a cvar.
+			cp_drawable_set_device_anchor(FrameState.SwiftDrawable, nullptr);
+			// Leave FrameState.HeadTransform untouched, same as game thread. It may simply be identity.  But whatever it is we will go ahead and render with it.
+		}
+	}
     
     //cp_drawable_set_ar_pose(FrameState.SwiftDrawable, )  // I think this is may be a deprecated way of specifying the frame's transform.  Replaced with cp_drawable_set_device_anchor
 }
@@ -1053,7 +1079,7 @@ void FOXRVisionOSSession::StartSubmission()
 XrResult FOXRVisionOSSession::XrEndFrame(
 	const XrFrameEndInfo* InFrameEndInfo)
 {
-	//UE_LOG(LogOXRVisionOS, Log, TEXT("FOXRVisionOSSession XrEndFrame RHI fliparg: %d BBIndex: %d"), PipelinedFrameStateRHI.FrameCounter, PipelinedFrameStateRHI.BackBufferIndex);
+	UE_LOG(LogOXRVisionOS, Verbose, TEXT("FOXRVisionOSSession::XrEndFrame"));
 
 	SCOPED_NAMED_EVENT_TEXT("FOXRVisionOSSession::XrEndFrame", FColor::Turquoise);
 
@@ -1081,6 +1107,10 @@ XrResult FOXRVisionOSSession::XrEndFrame(
         //TODO??? supporting only one layer now.  This means we do not have a non-reprojected layer.  Maybe we could have one as a separate swiftui?
         check(FrameEndInfo.layerCount == 1);
         check(FrameEndInfo.layers != nullptr);
+		if (FrameEndInfo.layers[0] == nullptr) {
+			return XrResult::XR_ERROR_LAYER_INVALID;
+		}
+		
         const XrCompositionLayerBaseHeader& LayerBaseHeader = *FrameEndInfo.layers[0];
         check(LayerBaseHeader.type == XR_TYPE_COMPOSITION_LAYER_PROJECTION);
         const XrCompositionLayerProjection& LayerProjection = *(reinterpret_cast<const XrCompositionLayerProjection*>(FrameEndInfo.layers[0]));
@@ -1092,9 +1122,17 @@ XrResult FOXRVisionOSSession::XrEndFrame(
         
         const FOXRVisionOSSwapchain& SwapchainData0 = *(reinterpret_cast<FOXRVisionOSSwapchain*>(View0.subImage.swapchain));
         const FOXRVisionOSSwapchain::FSwapchainImage& ReleasedImage = SwapchainData0.GetLastReleasedImage();
-        
-        //TODO feed in all necessary cp_frame_end_submission data
-        const MetalRHIVisionOS::PresentImmersiveParams Params{ReleasedImage.Image, PipelinedFrameStateRHI.SwiftDrawable};
+		
+		const XrCompositionLayerDepthInfoKHR* Depth0 = OpenXR::FindChainedStructByType<XrCompositionLayerDepthInfoKHR>(View0.next, (XrStructureType)XR_TYPE_COMPOSITION_LAYER_DEPTH_INFO_KHR);
+		const XrCompositionLayerDepthInfoKHR* Depth1 = OpenXR::FindChainedStructByType<XrCompositionLayerDepthInfoKHR>(View1.next, (XrStructureType)XR_TYPE_COMPOSITION_LAYER_DEPTH_INFO_KHR);
+		
+		check(Depth0 && Depth1);
+		check(Depth0->subImage.swapchain == Depth1->subImage.swapchain); //TODO: should we do it this way or should we have two separate swapchains?  If we did we would need to handle both below.
+
+		const FOXRVisionOSSwapchain& SwapchainDepthData0 = *(reinterpret_cast<FOXRVisionOSSwapchain*>(Depth0->subImage.swapchain));
+		const FOXRVisionOSSwapchain::FSwapchainImage& ReleasedDepth = SwapchainDepthData0.GetLastReleasedImage();
+
+        const MetalRHIVisionOS::PresentImmersiveParams Params{ReleasedImage.Image, ReleasedDepth.Image, PipelinedFrameStateRHI.SwiftFrame, PipelinedFrameStateRHI.SwiftDrawable};
         MetalRHIVisionOS::PresentImmersive(Params);
     }
 
diff --git a/Engine/Platforms/VisionOS/Plugins/Runtime/OpenXRVisionOS/Source/OXRVisionOS/Private/OXRVisionOSSession.h b/Engine/Platforms/VisionOS/Plugins/Runtime/OpenXRVisionOS/Source/OXRVisionOS/Private/OXRVisionOSSession.h
--- a/Engine/Platforms/VisionOS/Plugins/Runtime/OpenXRVisionOS/Source/OXRVisionOS/Private/OXRVisionOSSession.h
+++ b/Engine/Platforms/VisionOS/Plugins/Runtime/OpenXRVisionOS/Source/OXRVisionOS/Private/OXRVisionOSSession.h
@@ -229,8 +229,9 @@ private:
 	FPipelinedFrameState	PipelinedFrameStateRendering;
 	FPipelinedFrameState	PipelinedFrameStateRHI;
 	int32 CachedBeginFlipFrameCounter = INT32_MAX;
-
-	// TODO Reorganize this class?
+	
+	//TEMP see usages
+	FCriticalSection CriticalSection_ar_world_tracking_provider_query_device_anchor_at_timestamp;
 
 	FOpenXRRenderBridge* RenderBridge = nullptr;
 
diff --git a/Engine/Platforms/VisionOS/Plugins/Runtime/OpenXRVisionOS/Source/OXRVisionOS/Private/OXRVisionOSSwapchain.cpp b/Engine/Platforms/VisionOS/Plugins/Runtime/OpenXRVisionOS/Source/OXRVisionOS/Private/OXRVisionOSSwapchain.cpp
--- a/Engine/Platforms/VisionOS/Plugins/Runtime/OpenXRVisionOS/Source/OXRVisionOS/Private/OXRVisionOSSwapchain.cpp
+++ b/Engine/Platforms/VisionOS/Plugins/Runtime/OpenXRVisionOS/Source/OXRVisionOS/Private/OXRVisionOSSwapchain.cpp
@@ -66,25 +66,50 @@ FOXRVisionOSSwapchain::FOXRVisionOSSwapchain(const XrSwapchainCreateInfo* create
 	CreateInfo = *createInfo;
 
 	{
-		ETextureCreateFlags TextureCreateFlags = ETextureCreateFlags::None;
-		TextureCreateFlags |= TexCreate_RenderTargetable;
-		TextureCreateFlags |= TexCreate_UAV;
-		TextureCreateFlags |= TexCreate_ShaderResource;
-
-		check((CreateInfo.usageFlags & XR_SWAPCHAIN_USAGE_COLOR_ATTACHMENT_BIT) != 0);
-		check((CreateInfo.usageFlags & XR_SWAPCHAIN_USAGE_DEPTH_STENCIL_ATTACHMENT_BIT) == 0);
 		check(CreateInfo.faceCount == 1);
 		check(CreateInfo.arraySize == 1);
+		
+		const bool isColor = (CreateInfo.usageFlags & XR_SWAPCHAIN_USAGE_COLOR_ATTACHMENT_BIT) != 0;
+		const bool isDepth = (CreateInfo.usageFlags & XR_SWAPCHAIN_USAGE_DEPTH_STENCIL_ATTACHMENT_BIT) != 0;
+		check((isColor || isDepth) && (isColor != isDepth))
 
+		ETextureCreateFlags TextureCreateFlags = ETextureCreateFlags::None;
+		TextureCreateFlags |= TexCreate_RenderTargetable;
+		TextureCreateFlags |= TexCreate_ShaderResource;
+		if (isColor) 
+		{
+			TextureCreateFlags |= TexCreate_UAV;
+			TextureCreateFlags |= TexCreate_SRGB;
+		} 
+		else
+		{
+			if ((EPixelFormat)CreateInfo.format == PF_DepthStencil) 
+			{
+				TextureCreateFlags |= TexCreate_DepthStencilTargetable;
+				TextureCreateFlags |= TexCreate_DepthStencilResolveTarget;
+				TextureCreateFlags |= TexCreate_InputAttachmentRead;
+			}
+			else if ((EPixelFormat)CreateInfo.format == PF_R32_FLOAT)
+			{
+				TextureCreateFlags |= TexCreate_DepthStencilTargetable;
+				TextureCreateFlags |= TexCreate_ResolveTargetable;
+				TextureCreateFlags |= TexCreate_InputAttachmentRead;
+			}
+			else
+			{
+				check(false);
+			}
+		}
+		const FRHITextureCreateDesc BackBufferDesc =
+			FRHITextureCreateDesc::Create2D(isColor ? TEXT("OXRVisionOSSwapchain") : TEXT("OXRVisionOSSwapchainDepth"),
+											CreateInfo.width, CreateInfo.height,
+											(EPixelFormat)CreateInfo.format)
+			.SetFlags(TextureCreateFlags)
+			.DetermineInititialState();
+		
 		Images.Reserve(SwapchainLength);
 		for (int32 BufferIndex = 0; BufferIndex < SwapchainLength; ++BufferIndex)
 		{
-			// TODO: Add some identifying names?
-			const FRHITextureCreateDesc BackBufferDesc =
-				FRHITextureCreateDesc::Create2D(TEXT("OXRVisionOSSwapchain"), CreateInfo.width, CreateInfo.height, (EPixelFormat)CreateInfo.format)
-				.SetFlags(TextureCreateFlags)
-				.DetermineInititialState();
-				
 			Images.Emplace(RHICreateTexture(BackBufferDesc), EImageState::Released);
 		}
 	}
diff --git a/Engine/Platforms/VisionOS/Source/Runtime/Launch/Source/UESwift.swift b/Engine/Platforms/VisionOS/Source/Runtime/Launch/Source/UESwift.swift
--- a/Engine/Platforms/VisionOS/Source/Runtime/Launch/Source/UESwift.swift
+++ b/Engine/Platforms/VisionOS/Source/Runtime/Launch/Source/UESwift.swift
@@ -43,6 +43,40 @@ struct SwiftUIView: View {
   }
 }
 
+struct UEContentConfiguration: CompositorLayerConfiguration {
+	func makeConfiguration(
+		capabilities: LayerRenderer.Capabilities,
+		configuration: inout LayerRenderer.Configuration
+	)
+	{
+		/*
+		//let supportsFoveation = capabilites.supportsFoveation
+		//let supportedLayouts = supportedLayouts(options: supportsFoveation ? [.foveationEnabled] : [])
+		let supportedColorFormats = capabilities.supportedColorFormats
+		let supportedDepthFormats = capabilities.supportedDepthFormats
+		print("Supported Color Formats: ")
+		supportedColorFormats.forEach { colorFormat in
+			print(colorFormat.rawValue)
+		}
+		print("Supported Depth Formats: ")
+		supportedDepthFormats.forEach { depthFormat in
+			print(depthFormat.rawValue)
+		}
+		*/
+		
+		configuration.layout = .shared
+		//configuration.layout = .dedicated // separate texture for each eye.  We may need to switch when we implement foveated rendering.
+		configuration.isFoveationEnabled = false
+		
+		// HDR support // Might want to switch based on project settings.
+		//configuration.colorFormat = .rgba16Float
+		configuration.colorFormat = .bgra8Unorm_srgb
+		
+		//configuration.depthFormat = .depth32Float  			//PF_R32_FLOAT   			// This is correct for mobile forward
+		configuration.depthFormat = .depth32Float_stencil8 		//PF_DepthStencil   // This is correct for deferred
+//PFSWITCH
+	}
+}
 
 // unused at the moment, but this is how UE can open a SwiftUI view form the Obj-C side
 class HostingViewFactory: NSObject
@@ -72,7 +106,7 @@ struct UESwiftApp: App {
 		ImmersiveSpace(id: "ImmersiveSpace")
 		{
 			// this will make a CompositorLayer that can pull a Metal drawable out for UE to render to
-			CompositorLayer
+			CompositorLayer(configuration: UEContentConfiguration())
 			{
 				// on button click, tell Unreal the layer is ready and can continue
 				layerRenderer in
diff --git a/Engine/Source/Runtime/Apple/MetalRHI/Private/MetalRHIContext.cpp b/Engine/Source/Runtime/Apple/MetalRHI/Private/MetalRHIContext.cpp
--- a/Engine/Source/Runtime/Apple/MetalRHI/Private/MetalRHIContext.cpp
+++ b/Engine/Source/Runtime/Apple/MetalRHI/Private/MetalRHIContext.cpp
@@ -240,11 +240,3 @@ void FMetalRHICommandContext::RHIEndOcclusionQueryBatch()
 	check(CommandBufferFence.IsValid());
 	CommandBufferFence.Reset();
 }
-
-#if PLATFORM_VISIONOS
-void FMetalRHICommandContext::BeginRenderingImmersive(const MetalRHIVisionOS::BeginRenderingImmersiveParams& Params)
-{
-    SwiftFrame = Params.SwiftFrame;
-    UE_LOG(LogMetalVisionOS, Verbose, TEXT("SwiftLayerFrame(0x%x) copying from Params in FMetalRHICommandContext::BeginRenderingImmersive"), SwiftFrame);
-}
-#endif
diff --git a/Engine/Source/Runtime/Apple/MetalRHI/Private/MetalRHIVisionOSBridge.cpp b/Engine/Source/Runtime/Apple/MetalRHI/Private/MetalRHIVisionOSBridge.cpp
--- a/Engine/Source/Runtime/Apple/MetalRHI/Private/MetalRHIVisionOSBridge.cpp
+++ b/Engine/Source/Runtime/Apple/MetalRHI/Private/MetalRHIVisionOSBridge.cpp
@@ -9,14 +9,6 @@
 
 DEFINE_LOG_CATEGORY(LogMetalVisionOS);
 
-void MetalRHIVisionOS::BeginRenderingImmersive(const MetalRHIVisionOS::BeginRenderingImmersiveParams& Params)
-{
-    UE_LOG(LogMetalVisionOS, Verbose, TEXT("SwiftLayerFrame(0x%x) sending to RHICommandContext in MetalRHIVisionOS::BeginRendering"), Params.SwiftFrame);
-    FMetalRHICommandContext* RHICommandContext = static_cast<FMetalRHICommandContext*>(RHIGetDefaultContext());
-    check(RHICommandContext);
-    RHICommandContext->BeginRenderingImmersive(Params);
-}
-
 void MetalRHIVisionOS::PresentImmersive(const MetalRHIVisionOS::PresentImmersiveParams& Params)
 {
     FMetalRHICommandContext* RHICommandContext = static_cast<FMetalRHICommandContext*>(RHIGetDefaultContext());
@@ -26,6 +18,6 @@ void MetalRHIVisionOS::PresentImmersive(const MetalRHIVisionOS::PresentImmersive
     check(Viewport);
     FMetalSurface* Surface = GetMetalSurfaceFromRHITexture(Params.Texture);
 
-    Viewport->PresentImmersive(Surface, &Params);
+    Viewport->PresentImmersive(Surface, Params);
 }
 #endif // PLATFORM_VISIONOS
diff --git a/Engine/Source/Runtime/Apple/MetalRHI/Private/MetalViewport.cpp b/Engine/Source/Runtime/Apple/MetalRHI/Private/MetalViewport.cpp
--- a/Engine/Source/Runtime/Apple/MetalRHI/Private/MetalViewport.cpp
+++ b/Engine/Source/Runtime/Apple/MetalRHI/Private/MetalViewport.cpp
@@ -94,10 +94,6 @@ FMetalViewport::FMetalViewport(void* WindowHandle, uint32 InSizeX, uint32 InSize
 #if PLATFORM_VISIONOS
 	// look to see if we need to hook up to a Swift compositor renderer
 	SwiftLayer = [IOSAppDelegate GetDelegate].SwiftLayer;
-	
-	// Get an initial Frame
-	SwiftLayerFrame = cp_layer_renderer_query_next_frame(SwiftLayer);
-	UE_LOG(LogMetalVisionOS, Verbose, TEXT("SwiftLayerFrame(0x%x) from cp_layer_renderer_query_next_frame in MetalViewport constructor"), SwiftLayerFrame);
 #endif
 
 #if PLATFORM_MAC
@@ -391,24 +387,6 @@ CA::MetalDrawable* FMetalViewport::GetDrawable(EMetalViewportAccessFlag Accessor
 
 MTL::Texture* FMetalViewport::GetDrawableTexture(EMetalViewportAccessFlag Accessor)
 {
-#if PLATFORM_VISIONOS
-	if (SwiftLayerFrame != nullptr)
-	{
-		// get the next drawable
-		UE_LOG(LogMetalVisionOS, Verbose, TEXT("SwiftLayerFrame(0x%x) cp_frame_query_drawable in FMetalViewport::GetDrawableTexture"), SwiftLayerFrame);
-		cp_drawable_t SwiftDrawable = cp_frame_query_drawable(SwiftLayerFrame);
-		if (SwiftDrawable == nullptr)
-		{
-			UE_LOG(LogMetal, Fatal, TEXT("Failed to get next drawable"));
-		}
-		
-		// get the color texture out and use that with the RHI
-		uint32 Index = GetViewportIndex(Accessor);
-		DrawableTextures[Index] = (__bridge MTL::Texture*)cp_drawable_get_color_texture(SwiftDrawable, 0);
-		return DrawableTextures[Index];
-	}
-#endif
-	
 	CA::MetalDrawable* CurrentDrawable = GetDrawable(Accessor);
     uint32 Index = GetViewportIndex(Accessor);
     
@@ -494,14 +472,6 @@ void FMetalViewport::Present(FMetalCommandQueue& CommandQueue, bool bLockToVsync
 #endif
 			if (FrameAvailable > 0 && (InDisplayID == 0 || (DisplayID == InDisplayID && !bIsInLiveResize)))
 			{
-#if PLATFORM_VISIONOS
-				if (SwiftLayerFrame != nullptr)
-				{
-					// tell the compositor service we are going to submit soon and need the drawable
-					cp_frame_start_submission(SwiftLayerFrame);
-				}
-#endif
-
 				FPlatformAtomics::InterlockedDecrement(&FrameAvailable);
 				CA::MetalDrawable* LocalDrawable = GetDrawable(EMetalViewportAccessDisplayLink);
                 LocalDrawable->retain();
@@ -581,14 +551,6 @@ void FMetalViewport::Present(FMetalCommandQueue& CommandQueue, bool bLockToVsync
 #else // PLATFORM_MAC
 						CurrentCommandBuffer->GetMTLCmdBuffer()->addCompletedHandler(CommandBufferHandler);
 
-#if PLATFORM_VISIONOS
-						if (SwiftLayer != nullptr)
-						{
-                            id<MTLCommandBuffer> MTLCommandBuffer = (__bridge id<MTLCommandBuffer>)CurrentCommandBuffer->GetMTLCmdBuffer().get();
-							cp_drawable_encode_present(cp_frame_query_drawable(SwiftLayerFrame), MTLCommandBuffer);
-						}
-						else
-#endif
 						if (MinPresentDuration && GEnablePresentPacing)
 						{
                             LocalDrawable->presentAfterMinimumDuration(1.0f/(float)FramePace);
@@ -638,28 +600,35 @@ void FMetalViewport::Swap()
 		BackBuffer[0] = BB1;
 		BackBuffer[1] = BB0;
 	}
-	
-#if PLATFORM_VISIONOS
-	if (SwiftLayer && SwiftLayerFrame)
-	{
-		EndFrameImmersive();
-	}
-#endif
 }
 
 #if PLATFORM_VISIONOS
-void FMetalViewport::BeginRenderingImmersive(cp_frame_t SwiftFrame)
+void FMetalViewport::GetDrawableImmersiveTextures(EMetalViewportAccessFlag Accessor, cp_frame_t SwiftLayerFrame, MTL::Texture*& OutColorTexture, MTL::Texture*& OutDepthTexture)
 {
-	check(SwiftFrame);
-	SwiftLayerFrame = SwiftFrame;
-	UE_LOG(LogMetalVisionOS, Verbose, TEXT("SwiftLayerFrame(0x%x) set in FMetalViewport::BeginRenderingImmersive"), SwiftLayerFrame);
+	check(SwiftLayerFrame != nullptr);
+
+	// get the next SwiftDrawable
+	UE_LOG(LogMetalVisionOS, Verbose, TEXT("SwiftLayerFrame(0x%x) cp_frame_query_drawable in FMetalViewport::GetDrawableTexture"), SwiftLayerFrame);
+	cp_drawable_t SwiftDrawable = cp_frame_query_drawable(SwiftLayerFrame);
+	if (SwiftDrawable == nullptr)
+	{
+		UE_LOG(LogMetal, Fatal, TEXT("Failed to get next SwiftDrawable"));
+	}
+	
+	// get the color texture out and use that with the RHI
+	uint32 Index = GetViewportIndex(Accessor);
+	uint32 TextureCount = cp_drawable_get_texture_count(SwiftDrawable);
+	UE_LOG(LogMetalVisionOS, Verbose, TEXT("cp_drawable_get_texture_count %i"), TextureCount);
+	OutColorTexture = (__bridge MTL::Texture*)cp_drawable_get_color_texture(SwiftDrawable, 0);
+	OutDepthTexture = (__bridge MTL::Texture*)cp_drawable_get_depth_texture(SwiftDrawable, 0);
+	DrawableTextures[Index] = OutColorTexture;
 }
 
-// This is the present for Immersive visionos, through the OXRVisionOS plugin.
-void FMetalViewport::PresentImmersive(TRefCountPtr<FMetalSurface> CompleteFrame, const MetalRHIVisionOS::PresentImmersiveParams* VisionOSParams)
+// This is the present for Immersive visionOS, through the OXRVisionOS plugin.
+void FMetalViewport::PresentImmersive(TRefCountPtr<FMetalSurface> CompleteFrame, const MetalRHIVisionOS::PresentImmersiveParams& VisionOSParams)
 {
+	UE_LOG(LogMetalVisionOS, Verbose, TEXT("FMetalViewport::PresentImmersive"));
 	check(SwiftLayer);  // If no SwiftLayer we should not be trying to be immersive.
-	check(SwiftLayerFrame);  // BeginRenderingImmersive should have set this.
 
 	FMetalRHICommandContext* Context = static_cast<FMetalRHICommandContext*>(RHIGetDefaultContext());
 	FMetalDeviceContext& DeviceContext = GetMetalDeviceContext();
@@ -668,100 +637,102 @@ void FMetalViewport::PresentImmersive(TRefCountPtr<FMetalSurface> CompleteFrame,
 	FScopeLock Lock(&Mutex);
 
 	LastCompleteFrame = CompleteFrame;
+	TRefCountPtr<FMetalSurface> LastCompleteDepth = GetMetalSurfaceFromRHITexture(VisionOSParams.Depth);
 
-	FPlatformAtomics::InterlockedExchange(&FrameAvailable, 1);
-	
 	{
-		if (FrameAvailable > 0 )
+		MTL::Texture* DrawableTextureParam = nullptr;
+		MTL::Texture* DrawableDepthTextureParam = nullptr;
+		GetDrawableImmersiveTextures(EMetalViewportAccessDisplayLink, VisionOSParams.SwiftFrame, DrawableTextureParam, DrawableDepthTextureParam);
+		MTLTexturePtr DrawableTexture = NS::RetainPtr(DrawableTextureParam);
+		MTLTexturePtr DrawableDepthTexture = NS::RetainPtr(DrawableDepthTextureParam);
 		{
-			FPlatformAtomics::InterlockedDecrement(&FrameAvailable);
-            CA::MetalDrawable* LocalDrawable = GetDrawable(EMetalViewportAccessDisplayLink);
-            LocalDrawable->retain();
-			MTL::Texture* DrawableTexture = GetDrawableTexture(EMetalViewportAccessDisplayLink);
-
+			FScopeLock BlockLock(&Mutex);
+			
+			if (DrawableTexture)
 			{
-				FScopeLock BlockLock(&Mutex);
+				FMetalCommandBuffer* CurrentCommandBuffer = CommandQueue.CreateCommandBuffer();
+				check(CurrentCommandBuffer);
 				
-				if (DrawableTexture)
-				{
-					FMetalCommandBuffer* CurrentCommandBuffer = CommandQueue.CreateCommandBuffer();
-					check(CurrentCommandBuffer);
-					
 #if ENABLE_METAL_GPUPROFILE
-					FMetalProfiler* Profiler = FMetalProfiler::GetProfiler();
-					FMetalCommandBufferStats* Stats = Profiler->AllocateCommandBuffer(CurrentCommandBuffer->GetMTLCmdBuffer(), 0);
+				FMetalProfiler* Profiler = FMetalProfiler::GetProfiler();
+				FMetalCommandBufferStats* Stats = Profiler->AllocateCommandBuffer(CurrentCommandBuffer->GetMTLCmdBuffer(), 0);
 #endif
 
-					// TODO Currently we are using intermediate back buffer to connect the OXRVisionOS Swapchain to the drawable.
-					// I think we could use the drawable directly and avoid this copy.
-					check(GMetalSupportsIntermediateBackBuffer);
-					if (GMetalSupportsIntermediateBackBuffer)
+				// TODO Currently we are using intermediate back buffer to connect the OXRVisionOS Swapchain to the drawable.
+				// I think we could use the drawable directly and avoid this copy.
+				check(GMetalSupportsIntermediateBackBuffer);
+				if (GMetalSupportsIntermediateBackBuffer)
+				{
 					{
 						TRefCountPtr<FMetalSurface> Texture = LastCompleteFrame;
 						check(IsValidRef(Texture));
-						
-                        MTLTexturePtr Src = Texture->Texture;
-                        MTL::Texture* Dst = DrawableTexture;
+						MTLTexturePtr Src = Texture->Texture;
+						MTLTexturePtr& Dst = DrawableTexture;
 						
 						NSUInteger Width = FMath::Min(Src->width(), Dst->width());
 						NSUInteger Height = FMath::Min(Src->height(), Dst->height());
 						
 						MTL::BlitCommandEncoder* Encoder = CurrentCommandBuffer->GetMTLCmdBuffer()->blitCommandEncoder();
 						check(Encoder);
-
-                        METAL_GPUPROFILE(Profiler->EncodeBlit(Stats, __FUNCTION__));
-
-						Encoder->copyFromTexture(Src.get(), 0, 0, MTL::Origin(0, 0, 0), MTL::Size(Width, Height, 1), Dst, 0, 0, MTL::Origin(0, 0, 0));
+						METAL_GPUPROFILE(Profiler->EncodeBlit(Stats, __FUNCTION__));
+						Encoder->copyFromTexture(Src.get(), 0, 0, MTL::Origin(0, 0, 0), MTL::Size(Width, Height, 1), Dst.get(), 0, 0, MTL::Origin(0, 0, 0));
 						Encoder->endEncoding();
-
-						Drawable->release();
-						Drawable = nil;
 					}
-					
-					// This is a bit different than the usual pattern.
-					// This command buffer here is committed directly, instead of going through
-					// FMetalCommandList::Commit. So long as Present() is called within
-					// high level RHI BeginFrame/EndFrame this will not fine.
-					// Otherwise the recording of the Present time will be offset by one in the
-					// FMetalGPUProfiler frame indices.
-  
-                    MTL::HandlerFunction CommandBufferHandler = [LocalDrawable](MTL::CommandBuffer* CommandBuffer)
-					{
-						FMetalGPUProfiler::RecordPresent(CommandBuffer);
-						LocalDrawable->release();
-					};
-					
-					CurrentCommandBuffer->GetMTLCmdBuffer()->addCompletedHandler(CommandBufferHandler);
 
-					if (SwiftLayerFrame != nullptr)
 					{
-						if (VisionOSParams)
-						{
-							UE_LOG(LogMetalVisionOS, Verbose, TEXT("SwiftDrawable(0x%x) cp_drawable_encode_present in FMetalViewport::PresentImmersive"), VisionOSParams->SwiftDrawable);
-							cp_drawable_encode_present(VisionOSParams->SwiftDrawable, (__bridge id<MTLCommandBuffer>)CurrentCommandBuffer->GetMTLCmdBuffer().get());
-						}
+//HACK BEGIN
+// using LastCompleteFrame, the color result, is very wrong, but the depth texture is not yet plumbed through
+// on the mobile rendere and providing nothing results in an all black scene while this results in mostly-ok-pixels.
+// With the deferred renderer depth does already work, so flip this to using depth.
+#if 1
+						TRefCountPtr<FMetalSurface> Texture = LastCompleteFrame;
+#else
+						TRefCountPtr<FMetalSurface> Texture = LastCompleteDepth;
+#endif
+						//HACK END
+						check(IsValidRef(Texture));
+						MTLTexturePtr Src = Texture->Texture;
+						MTLTexturePtr& Dst = DrawableDepthTexture;
+						
+						NS::UInteger Width = FMath::Min(Src->width(), Dst->width());
+						NS::UInteger Height = FMath::Min(Src->height(), Dst->height());
+						
+						MTLBlitCommandEncoderPtr Encoder = NS::RetainPtr(CurrentCommandBuffer->GetMTLCmdBuffer()->blitCommandEncoder());
+						check(Encoder);
+						METAL_GPUPROFILE(Profiler->EncodeBlit(Stats, __FUNCTION__));
+						Encoder->copyFromTexture(Src.get(), 0, 0, MTL::Origin(0, 0, 0), MTL::Size(Width, Height, 1), Dst.get(), 0, 0, MTL::Origin(0, 0, 0));
+						Encoder->endEncoding();
 					}
+				}
+				
+				// This is a bit different than the usual pattern.
+				// This command buffer here is committed directly, instead of going through
+				// FMetalCommandList::Commit. So long as Present() is called within
+				// high level RHI BeginFrame/EndFrame this will not fine.
+				// Otherwise the recording of the Present time will be offset by one in the
+				// FMetalGPUProfiler frame indices.
+
+				MTL::HandlerFunction CommandBufferHandler = [](MTL::CommandBuffer* CommandBuffer)
+				{
+					FMetalGPUProfiler::RecordPresent(CommandBuffer);
+				};
 
-					METAL_GPUPROFILE(Stats->End(CurrentCommandBuffer->GetMTLCmdBuffer()));
-					CommandQueue.CommitCommandBuffer(CurrentCommandBuffer);
+				if (VisionOSParams.SwiftFrame != nullptr) // TODO maybe remove this?
+				{
+					UE_LOG(LogMetalVisionOS, Verbose, TEXT("SwiftDrawable(0x%x) cp_drawable_encode_present in FMetalViewport::PresentImmersive"), VisionOSParams.SwiftDrawable);
+					cp_drawable_encode_present(VisionOSParams.SwiftDrawable, (__bridge id<MTLCommandBuffer>)CurrentCommandBuffer->GetMTLCmdBuffer().get());
 				}
+				
+				CurrentCommandBuffer->GetMTLCmdBuffer()->addCompletedHandler(CommandBufferHandler);
+
+				METAL_GPUPROFILE(Stats->End(CurrentCommandBuffer->GetMTLCmdBuffer()));
+				CommandQueue.CommitCommandBuffer(CurrentCommandBuffer);
 			}
 		}
 	}
 	
-	EndFrameImmersive();
-}
-
-void FMetalViewport::EndFrameImmersive()
-{
-	check(SwiftLayer);
-	check(SwiftLayerFrame);
-	
-	// end the last frame
-	UE_LOG(LogMetalVisionOS, Verbose, TEXT("SwiftLayerFrame(0x%x) cp_frame_end_submission in FMetalViewport::EndFrameImmersive"), SwiftLayerFrame);
-	cp_frame_end_submission(SwiftLayerFrame);
-	UE_LOG(LogMetalVisionOS, Verbose, TEXT("SwiftLayerFrame nulled in FMetalViewport::EndFrame()"));
-	SwiftLayerFrame = nullptr;
+	UE_LOG(LogMetalVisionOS, Verbose, TEXT("FMetalViewport::PresentImmersive calling cp_frame_end_submission with VisionOSParams.SwiftFrame(0x%x)"), VisionOSParams.SwiftFrame);
+	cp_frame_end_submission(VisionOSParams.SwiftFrame);
 }
 #endif
 
@@ -810,13 +781,6 @@ void FMetalRHIImmediateCommandContext::RHIBeginDrawingViewport(FRHIViewport* Vie
     MTL_SCOPED_AUTORELEASE_POOL;
 	FMetalViewport* Viewport = ResourceCast(ViewportRHI);
 	check(Viewport);
-	 
-#if PLATFORM_VISIONOS
-		if (SwiftFrame)
-		{
-			Viewport->BeginRenderingImmersive(SwiftFrame);
-		}
-#endif
 		
 	((FMetalDeviceContext*)Context)->BeginDrawingViewport(Viewport);
 
diff --git a/Engine/Source/Runtime/Apple/MetalRHI/Public/MetalRHIVisionOSBridge.h b/Engine/Source/Runtime/Apple/MetalRHI/Public/MetalRHIVisionOSBridge.h
--- a/Engine/Source/Runtime/Apple/MetalRHI/Public/MetalRHIVisionOSBridge.h
+++ b/Engine/Source/Runtime/Apple/MetalRHI/Public/MetalRHIVisionOSBridge.h
@@ -12,15 +12,11 @@ DECLARE_LOG_CATEGORY_EXTERN(LogMetalVisionOS, Warning, All);
 
 namespace MetalRHIVisionOS
 {
-	METALRHI_API struct BeginRenderingImmersiveParams
-    {
-        const cp_frame_t SwiftFrame;
-    };
-    METALRHI_API void BeginRenderingImmersive(const BeginRenderingImmersiveParams& Params);
-
     METALRHI_API struct PresentImmersiveParams
     {
         const FTextureRHIRef& Texture;
+		const FTextureRHIRef& Depth;
+		cp_frame_t& SwiftFrame;
         cp_drawable_t& SwiftDrawable;
     };
     METALRHI_API void PresentImmersive(const PresentImmersiveParams& Params);
diff --git a/Engine/Source/Runtime/Apple/MetalRHI/Public/MetalViewport.h b/Engine/Source/Runtime/Apple/MetalRHI/Public/MetalViewport.h
--- a/Engine/Source/Runtime/Apple/MetalRHI/Public/MetalViewport.h
+++ b/Engine/Source/Runtime/Apple/MetalRHI/Public/MetalViewport.h
@@ -78,10 +78,8 @@ public:
 	void Swap();
 	
 #if PLATFORM_VISIONOS
-    void BeginRenderingImmersive(cp_frame_t SwiftFrame);
-    void PresentImmersive(TRefCountPtr<FMetalSurface> CompleteFrame, const MetalRHIVisionOS::PresentImmersiveParams* Params);
-private:
-	void EndFrameImmersive();
+	void GetDrawableImmersiveTextures(EMetalViewportAccessFlag Accessor, cp_frame_t SwiftLayerFrame, MTL::Texture*& OutColorTexture, MTL::Texture*& OutDepthTexture );
+    void PresentImmersive(TRefCountPtr<FMetalSurface> CompleteFrame, const MetalRHIVisionOS::PresentImmersiveParams& Params);
 #endif
 	
 private:
@@ -90,7 +88,6 @@ private:
 private:
 #if PLATFORM_VISIONOS
 	CP_OBJECT_cp_layer_renderer* SwiftLayer = nullptr;
-	cp_frame_t SwiftLayerFrame = nullptr;
 #endif
 	
 	CA::MetalDrawable* Drawable;
diff --git a/Engine/Source/Runtime/Launch/Private/IOS/LaunchIOS.cpp b/Engine/Source/Runtime/Launch/Private/IOS/LaunchIOS.cpp
--- a/Engine/Source/Runtime/Launch/Private/IOS/LaunchIOS.cpp
+++ b/Engine/Source/Runtime/Launch/Private/IOS/LaunchIOS.cpp
@@ -621,19 +621,13 @@ void FSwiftAppBootstrap::KickoffWithCompositingLayer(CP_OBJECT_cp_layer_renderer
 			cp_view_texture_map_t TextureMap = cp_view_get_view_texture_map(View);
 			MTLViewport Viewport = cp_view_texture_map_get_viewport(TextureMap);
 
-			float ScaleDownHack = 0.5;
-			float X = Viewport.originX * ScaleDownHack;
-			float Y = Viewport.originY * ScaleDownHack;
-			float W = Viewport.width * ScaleDownHack;
-			float H = Viewport.height * ScaleDownHack;
-			// both views are at 0? if so, hack the view 1 to be to the right of view 0
-			if (ViewIndex == 1 && X == 0)
-			{
-				X = W;
-			}
+			float X = Viewport.originX;
+			float Y = Viewport.originY;
+			float W = Viewport.width;
+			float H = Viewport.height;
 			NSValue* VPValue = [NSValue valueWithCGRect:CGRectMake(X, Y, W, H)];
 			[Viewports addObject:VPValue];
-			FPlatformMisc::LowLevelOutputDebugStringf(TEXT("Adding eye viewport : [%f. %f] / [%f x %f] (scaled by %f\n"), X, Y, W, H, ScaleDownHack);
+			FPlatformMisc::LowLevelOutputDebugStringf(TEXT("Adding eye viewport : [%f. %f] / [%f x %f]\n"), X, Y, W, H);
 		}
 		cp_frame_start_submission(SwiftLayerFrame);
 		id<MTLDevice> Device = cp_layer_renderer_get_device(Layer);
