UE-208427 VisionOS: problem with view transforms
-Fixed bug where hmd offset was applied twice.
-Fixed bug where eye offsets were never applied.
-Removed some unused code and obsolete comments.
-Additional logging to debug frame sequence issues.

diff --git a/Engine/Platforms/VisionOS/Plugins/Runtime/OpenXRVisionOS/Source/OXRVisionOS/Private/OXRVisionOS.h b/Engine/Platforms/VisionOS/Plugins/Runtime/OpenXRVisionOS/Source/OXRVisionOS/Private/OXRVisionOS.h
--- a/Engine/Platforms/VisionOS/Plugins/Runtime/OpenXRVisionOS/Source/OXRVisionOS/Private/OXRVisionOS.h
+++ b/Engine/Platforms/VisionOS/Plugins/Runtime/OpenXRVisionOS/Source/OXRVisionOS/Private/OXRVisionOS.h
@@ -18,6 +18,26 @@
 DECLARE_LOG_CATEGORY_EXTERN(LogOXRVisionOS, Log, All);
 // Note: if you are turning this log up to verbose you may also want LogMetalVisionOS turned up to verbose.
 
+static const TCHAR* VOSThreadString()
+{
+	if (IsInGameThread()) 
+	{
+		return TEXT("GT~");
+	}
+	else if (IsInRenderingThread())
+	{
+		return TEXT("RT~");
+	}
+	else if (IsInRHIThread())
+	{
+		return TEXT("IT~");
+	}
+	else
+	{
+		return TEXT("?T~");
+	}
+}
+
 class FOXRVisionOSInstance;
 
 class FOXRVisionOS : public IModuleInterface, public IOpenXRExtensionPlugin
diff --git a/Engine/Platforms/VisionOS/Plugins/Runtime/OpenXRVisionOS/Source/OXRVisionOS/Private/OXRVisionOSPlatformUtils.h b/Engine/Platforms/VisionOS/Plugins/Runtime/OpenXRVisionOS/Source/OXRVisionOS/Private/OXRVisionOSPlatformUtils.h
--- a/Engine/Platforms/VisionOS/Plugins/Runtime/OpenXRVisionOS/Source/OXRVisionOS/Private/OXRVisionOSPlatformUtils.h
+++ b/Engine/Platforms/VisionOS/Plugins/Runtime/OpenXRVisionOS/Source/OXRVisionOS/Private/OXRVisionOSPlatformUtils.h
@@ -183,10 +183,12 @@ namespace OXRVisionOS
         Ret.orientation.y = RawRotation.Y;
         Ret.orientation.z = RawRotation.Z;
         Ret.orientation.w = RawRotation.W;
-        Ret.position.x = RawYUpFMatrix.M[3][0];
-        Ret.position.y = RawYUpFMatrix.M[3][1];
-        Ret.position.z = RawYUpFMatrix.M[3][2];
-        return Ret;
+
+		Ret.position.x = RawYUpMatrix.columns[3][0];
+		Ret.position.y = RawYUpMatrix.columns[3][1];
+		Ret.position.z = RawYUpMatrix.columns[3][2];
+		
+		return Ret;
     }
 	FORCEINLINE XrPosef ToXrPose(const FTransform& InTransform)
 	{
@@ -213,7 +215,7 @@ namespace OXRVisionOS
 
 	FORCEINLINE XrSpaceVelocityFlags ToXrSpaceVelocityFlags(TrackerResultData& InResultData)
 	{
-//TODO finish this
+//TODO finish this, or delete it
 		XrSpaceVelocityFlags OutFlags = 0;
 
 		// 	OutFlags |= XR_SPACE_VELOCITY_LINEAR_VALID_BIT;
diff --git a/Engine/Platforms/VisionOS/Plugins/Runtime/OpenXRVisionOS/Source/OXRVisionOS/Private/OXRVisionOSSession.cpp b/Engine/Platforms/VisionOS/Plugins/Runtime/OpenXRVisionOS/Source/OXRVisionOS/Private/OXRVisionOSSession.cpp
--- a/Engine/Platforms/VisionOS/Plugins/Runtime/OpenXRVisionOS/Source/OXRVisionOS/Private/OXRVisionOSSession.cpp
+++ b/Engine/Platforms/VisionOS/Plugins/Runtime/OpenXRVisionOS/Source/OXRVisionOS/Private/OXRVisionOSSession.cpp
@@ -93,6 +93,10 @@ FOXRVisionOSSession::FOXRVisionOSSession(const XrSessionCreateInfo* createInfo,
 		bCreateFailed = true;
 		return;
 	}
+	
+	for (int32 i = 0; i < RenderToGameHeadTransformBufferLength; ++i) {
+		RenderToGameHeadTransform[i] = matrix_identity_float4x4;
+	}
 
 	//if (FOXRVisionOSController::Create(Controllers, this) == false)
 	//{
@@ -239,17 +243,6 @@ XrResult FOXRVisionOSSession::XrEnumerateReferenceSpaces(
 		OutSpaces[0] = XR_REFERENCE_SPACE_TYPE_VIEW;
 		OutSpaces[1] = XR_REFERENCE_SPACE_TYPE_LOCAL;
 		//OutSpaces[2] = XR_REFERENCE_SPACE_TYPE_STAGE;
-
-		//CoordinateSystemType CurrentCoordinateSystem = COORDINATE_SYSTEM_LOCAL_FLOOR;
-		// const int32 Ret = SetCoordinateSystem(COORDINATE_SYSTEM_STAGE);
-		// if (Ret == RESULT_OK)
-		// {
-		// 	UE_LOG(LogOXRVisionOS, Log, TEXT("XrEnumerateReferenceSpaces called SetCoordinateSystem(COORDINATE_SYSTEM_STAGE)."));
-		// }
-		// else
-		// {
-		// 	UE_LOG(LogOXRVisionOS, Warning, TEXT("XrEnumerateReferenceSpaces call to SetCoordinateSystem(COORDINATE_SYSTEM_STAGE) failed!"));
-		// }
 	}
 	return XrResult::XR_SUCCESS;
 }
@@ -319,6 +312,18 @@ XrResult FOXRVisionOSSession::XrCreateActionSpace(
 	return Ret;
 }
 
+const FOXRVisionOSSpace* FOXRVisionOSSession::ToFOXRVisionOSSpace(const XrSpace& Space)
+{
+	// Check that this is a space we know about.
+	uint32 ArrayIndex = Spaces.IndexOfByPredicate([Space](const TSharedPtr<FOXRVisionOSSpace, ESPMode::ThreadSafe>& Data) { return (Data.Get() == (FOXRVisionOSSpace*)Space); });
+	if (ArrayIndex == INDEX_NONE)
+	{
+		return nullptr;
+	}
+	
+	return (FOXRVisionOSSpace*)Space;
+}
+
 XrResult FOXRVisionOSSession::DestroySpace(FOXRVisionOSSpace* Space)
 {
 	uint32 ArrayIndex = Spaces.IndexOfByPredicate([Space](const TSharedPtr<FOXRVisionOSSpace, ESPMode::ThreadSafe>& Data) { return (Data.Get() == Space); });
@@ -853,25 +858,16 @@ XrResult FOXRVisionOSSession::XrWaitFrame(
 	// Block until the previous frame's xrBeginFrame has happened.
 	{
 		SCOPED_NAMED_EVENT_TEXT("FOXRVisionOSSession::XrWaitFrame XrWaitFrameEvent->Wait()", FColor::Turquoise);
-		UE_LOG(LogOXRVisionOS, Verbose, TEXT("XrWaitFrameEvent->Wait() Started"));
+		UE_LOG(LogOXRVisionOS, Verbose, TEXT("XrWaitFrameEvent->Wait() Started   FC will be %i"), SessionFrameCounter);
 		XrWaitFrameEvent->Wait();
 	}
 
 	PipelinedFrameStateGame.FrameCounter = SessionFrameCounter++;
-	UE_LOG(LogOXRVisionOS, Verbose, TEXT("XrWaitFrameEvent->Wait() Finished 0x%x"), PipelinedFrameStateGame.SwiftFrame);
-
-	//UE_LOG(LogOXRVisionOS, Log, TEXT("FOXRVisionOSSession XrWaitFrame proceeds Game fliparg: %d BBIndex: %d"), PipelinedFrameStateGame.FrameCounter, PipelinedFrameStateGame.BackBufferIndex);
-
-	// {
-	// 	const int32 ret = ReprojectionBeginFrame(PipelinedFrameStateGame.FrameCounter);
-	// 	if (ret != RESULT_OK)
-	// 	{
-	// 		UE_LOG(LogOXRVisionOS, Warning, TEXT("ReprojectionBeginFrame(%d) failed: %d (%u, 0x%08X)"), PipelinedFrameStateGame.FrameCounter, ret, ret, ret);
-	// 	}
-	// }
+	PipelinedFrameStateGame.RenderToGameHeadTransformIndexRead++;
+	UE_LOG(LogOXRVisionOS, Verbose, TEXT("%s FC: %i XrWaitFrameEvent->Wait() Finished"), VOSThreadString(), PipelinedFrameStateGame.FrameCounter);
     
     PipelinedFrameStateGame.SwiftFrame = cp_layer_renderer_query_next_frame(VOSLayerRenderer);
-    UE_LOG(LogOXRVisionOS, Verbose, TEXT("SwiftLayerFrame(0x%x) from cp_layer_renderer_query_next_frame in FOXRVisionOSSession::XrWaitFrame"), PipelinedFrameStateGame.SwiftFrame);
+    UE_LOG(LogOXRVisionOS, Verbose, TEXT("%s FC: %i SLF=0x%x cp_layer_renderer_query_next_frame in FOXRVisionOSSession::XrWaitFrame SwiftLayerFrame fetched"), VOSThreadString(), PipelinedFrameStateGame.FrameCounter, PipelinedFrameStateGame.SwiftFrame);
 
     // Fetch the predicted timing information.
     PipelinedFrameStateGame.SwiftFrameTiming = cp_frame_predict_timing(PipelinedFrameStateGame.SwiftFrame);
@@ -885,20 +881,22 @@ XrResult FOXRVisionOSSession::XrWaitFrame(
     }
 
     // Mark the beginnng of our game thread update work.  All the stuff before rendering.
-    UE_LOG(LogOXRVisionOS, Verbose, TEXT("SwiftLayerFrame(0x%x) from cp_frame_start_update in FOXRVisionOSSession::XrWaitFrame"), PipelinedFrameStateGame.SwiftFrame);
+    UE_LOG(LogOXRVisionOS, Verbose, TEXT("%s FC: %i SLF: 0x%x cp_frame_start_update in FOXRVisionOSSession::XrWaitFrame"), VOSThreadString(), PipelinedFrameStateGame.FrameCounter, PipelinedFrameStateGame.SwiftFrame);
     cp_frame_start_update(PipelinedFrameStateGame.SwiftFrame);
 
     const cp_time_t PresentationTime = cp_frame_timing_get_presentation_time(PipelinedFrameStateGame.SwiftFrameTiming);
     const CFTimeInterval TargetRenderTime = cp_time_to_cf_time_interval(PresentationTime); // CRTimeInterval is double time in seconds.
-    
+
 	XrTime DisplayPeriod = 6 * 1000000; //TODO other frame sequences?  This is 90fps
 	frameState->predictedDisplayPeriod = DisplayPeriod;
 	frameState->predictedDisplayTime = OXRVisionOS::APITimeToXrTime(TargetRenderTime);
 	PipelinedFrameStateGame.PredictedDisplayTime = frameState->predictedDisplayTime;
-    
-    PipelinedFrameStateGame.DeviceAnchor = ar_device_anchor_create();
-    check(ARKitWorldTrackingProvider)
-	UE_LOG(LogOXRVisionOS, Verbose, TEXT("SwiftLayerFrame TargetRenderTime= %f"), (double)TargetRenderTime);
+	
+	UE_LOG(LogOXRVisionOS, Verbose, TEXT("%s FC: %i SLF: 0x%x cp_frame_timing_get_presentation_time in FOXRVisionOSSession::XrWaitFrame TargetRenderTime =  %f xr predictedDisplayTime=%f"), VOSThreadString(), PipelinedFrameStateGame.FrameCounter, PipelinedFrameStateGame.SwiftFrame, (double)TargetRenderTime, (double)frameState->predictedDisplayTime);
+
+
+	PipelinedFrameStateGame.DeviceAnchor = ar_device_anchor_create();
+	check(ARKitWorldTrackingProvider)
 	{
 		//TODO: use a previous frame transform here, because we can't call this function from multiple threads
 		//TEMP use a lock.
@@ -915,7 +913,9 @@ XrResult FOXRVisionOSSession::XrWaitFrame(
 			PipelinedFrameStateGame.HeadTransform = matrix_identity_float4x4;
 		}
 	}
-    
+
+//	PipelinedFrameStateGame.HeadTransform = RenderToGameHeadTransformRead(PipelinedFrameStateGame);
+	
     LocateViewInfoIndexAdvance();
 
 	frameState->shouldRender = SessionState == XR_SESSION_STATE_FOCUSED || SessionState == XR_SESSION_STATE_VISIBLE;
@@ -923,6 +923,7 @@ XrResult FOXRVisionOSSession::XrWaitFrame(
 	// We have to pipeline state updates in order to prevent task overlap (esp during map loading where we kick extra renders)
 	ENQUEUE_RENDER_COMMAND(UpdatePipelinedFrameState)([this, NewRenderingFrameState = PipelinedFrameStateGame](FRHICommandListImmediate& RHICmdList)
 	{
+		UE_LOG(LogOXRVisionOS, Verbose, TEXT("%s FC: %i SLF: 0x%x UpdatePipelinedFrameState G->R "), VOSThreadString(), NewRenderingFrameState.FrameCounter, NewRenderingFrameState.SwiftFrame);
 		check(IsInRenderingThread());
 		PipelinedFrameStateRendering = NewRenderingFrameState;
 	});
@@ -935,10 +936,10 @@ XrResult FOXRVisionOSSession::XrWaitFrame(
 
 void FOXRVisionOSSession::OnBeginRendering_GameThread()
 {
-    // Mark the end of our previous game frame
+    // Mark the end of our game frame
     if (PipelinedFrameStateGame.SwiftFrame)
     {
-        UE_LOG(LogOXRVisionOS, VeryVerbose, TEXT("SwiftLayerFrame(0x%x) from cp_frame_end_update in FOXRVisionOSSession::OnBeginRendering_GameThread"), PipelinedFrameStateGame.SwiftFrame);
+		UE_LOG(LogOXRVisionOS, Verbose, TEXT("%s FC: %i SLF: 0x%x cp_frame_end_update in FOXRVisionOSSession::OnBeginRendering_GameThread"), VOSThreadString(), PipelinedFrameStateGame.FrameCounter, PipelinedFrameStateGame.SwiftFrame);
         cp_frame_end_update(PipelinedFrameStateGame.SwiftFrame);
         PipelinedFrameStateGame.SwiftFrame = nullptr;
     }
@@ -947,13 +948,19 @@ void FOXRVisionOSSession::OnBeginRendering_GameThread()
 // FOpenXRHMD::OnBeginRendering_RenderThread
 void FOXRVisionOSSession::OnBeginRendering_RenderThread()
 {
-    WaitUntil();
+	UE_LOG(LogOXRVisionOS, VeryVerbose, TEXT("%s FC: %i SLF: 0x%x FOXRVisionOSSession::OnBeginRendering_RenderThread"), VOSThreadString(), PipelinedFrameStateRendering.FrameCounter, PipelinedFrameStateRendering.SwiftFrame);
+	
+	WaitUntil();
     StartSubmission();
 	
+	UE_LOG(LogOXRVisionOS, VeryVerbose, TEXT("Enqueueing framestate from render to rhi: NewFrameState.FrameCounter %i"), PipelinedFrameStateRendering.FrameCounter);
 	FRHICommandListImmediate& RHICmdList = FRHICommandListImmediate::Get();
 	RHICmdList.EnqueueLambda([this, NewRHIFrameState = PipelinedFrameStateRendering](FRHICommandListImmediate&)
 	{
+		UE_LOG(LogOXRVisionOS, VeryVerbose, TEXT("%s FC: %i SLF: 0x%x UpdatePipelinedFrameStat R->I, but HeadTransform Cleared"), VOSThreadString(), NewRHIFrameState.FrameCounter, NewRHIFrameState.SwiftFrame);
+		
 		PipelinedFrameStateRHI = NewRHIFrameState;
+		PipelinedFrameStateRHI.HeadTransform = matrix_identity_float4x4; // overwrite to ensure rhi gets a new one before it uses one.
 	});
 }
 
@@ -963,9 +970,9 @@ XrResult FOXRVisionOSSession::XrBeginFrame(
 {
 	check(!IsRHIThreadRunning() || IsInRHIThread());
 
-	//UE_LOG(LogOXRVisionOS, Log, TEXT("FOXRVisionOSSession XrBeginFrame fliparg: %d BBIndex: %d"), PipelinedFrameStateRendering.FrameCounter, PipelinedFrameStateRendering.BackBufferIndex); // Note frame state not copied yet
-
 	SCOPED_NAMED_EVENT_TEXT("FOXRVisionOSSession::XrBeginFrame", FColor::Turquoise);
+	
+	UE_LOG(LogOXRVisionOS, VeryVerbose, TEXT("%s FC: %i SLF: 0x%x XrBeginFrame"), VOSThreadString(), PipelinedFrameStateRHI.FrameCounter, PipelinedFrameStateRHI.SwiftFrame);
 
 	// frameBeginInfo can be null
 	if (FrameBeginInfo != nullptr && FrameBeginInfo->type != XR_TYPE_FRAME_BEGIN_INFO)
@@ -994,7 +1001,6 @@ XrResult FOXRVisionOSSession::XrBeginFrame(
 	if (PipelinedFrameStateRHI.bSynchronizing)
 	{
 		XrWaitFrameEvent->Trigger();
-		//UE_LOG(LogOXRVisionOS, Log, TEXT("FOXRVisionOSSession XrBeginFrame bSynchronizing=true RHI fliparg: %d BBIndex: %d"), PipelinedFrameStateRHI.FrameCounter, PipelinedFrameStateRHI.BackBufferIndex);
 		return XrResult::XR_SUCCESS;
 	}
 
@@ -1016,10 +1022,10 @@ XrResult FOXRVisionOSSession::XrBeginFrame(
     {
 		FPipelinedFrameState& FrameState = PipelinedFrameStateRHI;
         
-        UE_LOG(LogOXRVisionOS, VeryVerbose, TEXT("SwiftLayerFrame(0x%x) cp_frame_query_drawable in FOXRVisionOSSession::XrBeginFrame"), FrameState.SwiftFrame);
         FrameState.SwiftDrawable = cp_frame_query_drawable(FrameState.SwiftFrame);
         check(FrameState.SwiftDrawable);
-        UE_LOG(LogOXRVisionOS, VeryVerbose, TEXT("SwiftLayerFrame(0x%x) cp_frame_query_drawable got SwiftDrawable(0x%x) in FOXRVisionOSSession::StartSubmission()"), FrameState.SwiftFrame, FrameState.SwiftDrawable);
+		UE_LOG(LogOXRVisionOS, VeryVerbose, TEXT("%s FC: %i SLF: 0x%x SD: 0x%x cp_frame_query_drawable in FOXRVisionOSSession::XrBeginFrame SwiftDrawable fetched"), VOSThreadString(), FrameState.FrameCounter, FrameState.SwiftFrame, FrameState.SwiftDrawable);
+
     }
     
     //MetalRHIOXRVisionOS::WaitUntilSafeForRendering(PipelinedFrameStateRHI.CommandListContext, NextBackBufferIndex);
@@ -1032,27 +1038,34 @@ void FOXRVisionOSSession::WaitUntil()
 {
     FPipelinedFrameState& FrameState = GetPipelinedFrameStateForThread();
     
-    UE_LOG(LogOXRVisionOS, VeryVerbose, TEXT("SwiftLayerFrame(0x%x) cp_time_wait_until starting in FOXRVisionOSSession::WaitUntil"), FrameState.SwiftFrame);
+	UE_LOG(LogOXRVisionOS, VeryVerbose, TEXT("%s FC: %i SLF: 0x%x SD: 0x%x cp_time_wait_until in FOXRVisionOSSession::WaitUntil started"), VOSThreadString(), FrameState.FrameCounter, FrameState.SwiftFrame, FrameState.SwiftDrawable);
+	
     cp_time_wait_until(cp_frame_timing_get_optimal_input_time(FrameState.SwiftFrameTiming));
-    UE_LOG(LogOXRVisionOS, VeryVerbose, TEXT("SwiftLayerFrame(0x%x) cp_time_wait_until completed in FOXRVisionOSSession::WaitUntil"), FrameState.SwiftFrame);
+
+	UE_LOG(LogOXRVisionOS, VeryVerbose, TEXT("%s FC: %i SLF: 0x%x SD: 0x%x cp_time_wait_until in FOXRVisionOSSession::WaitUntil completed"), VOSThreadString(), FrameState.FrameCounter, FrameState.SwiftFrame, FrameState.SwiftDrawable);
 }
 
 void FOXRVisionOSSession::StartSubmission()
 {
     FPipelinedFrameState& FrameState = GetPipelinedFrameStateForThread();
+	
+	UE_LOG(LogOXRVisionOS, VeryVerbose, TEXT("%s FC: %i SLF: 0x%x SD: 0x%x cp_frame_start_submission in FOXRVisionOSSession::StartSubmission"), VOSThreadString(), FrameState.FrameCounter, FrameState.SwiftFrame, FrameState.SwiftDrawable);
     
-    // tell the compositor service we are going to submit soon and need the drawable
-    UE_LOG(LogOXRVisionOS, VeryVerbose, TEXT("SwiftLayerFrame(0x%x) cp_frame_start_submission in FOXRVisionOSSession::XrBeginFrame"), FrameState.SwiftFrame);
     cp_frame_start_submission(FrameState.SwiftFrame);
 
-    UE_LOG(LogOXRVisionOS, VeryVerbose, TEXT("SwiftLayerFrame(0x%x) cp_frame_query_drawable in FOXRVisionOSSession::XrBeginFrame"), FrameState.SwiftFrame);
     FrameState.SwiftDrawable = cp_frame_query_drawable(FrameState.SwiftFrame);
     check(FrameState.SwiftDrawable);
-    UE_LOG(LogOXRVisionOS, VeryVerbose, TEXT("SwiftLayerFrame(0x%x) cp_frame_query_drawable got SwiftDrawable(0x%x) in FOXRVisionOSSession::StartSubmission()"), FrameState.SwiftFrame, FrameState.SwiftDrawable);
+	
+	UE_LOG(LogOXRVisionOS, VeryVerbose, TEXT("%s FC: %i SLF: 0x%x SD: 0x%x cp_frame_query_drawable in FOXRVisionOSSession::StartSubmission got the SD again."), VOSThreadString(), FrameState.FrameCounter, FrameState.SwiftFrame, FrameState.SwiftDrawable);
 
     FrameState.SwiftFinalFrameTiming = cp_drawable_get_frame_timing(FrameState.SwiftDrawable);
-    
+	
+	//FrameState.DeviceAnchor = ar_device_anchor_create();
+	
     CFTimeInterval TimeInterval = cp_time_to_cf_time_interval(cp_frame_timing_get_presentation_time(FrameState.SwiftFinalFrameTiming));
+	
+	UE_LOG(LogOXRVisionOS, VeryVerbose, TEXT("%s FC: %i SLF: 0x%x SD: 0x%x cp_drawable_get_frame_timing in FOXRVisionOSSession::StartSubmission got SwiftFinalFrameTiming %f"), VOSThreadString(), FrameState.FrameCounter, FrameState.SwiftFrame, FrameState.SwiftDrawable, (double)TimeInterval);
+	
     check(ARKitWorldTrackingProvider)
 	{
 		//TEMP use a lock.
@@ -1060,26 +1073,38 @@ void FOXRVisionOSSession::StartSubmission()
 		auto anchor_status = ar_world_tracking_provider_query_device_anchor_at_timestamp(ARKitWorldTrackingProvider, TimeInterval, FrameState.DeviceAnchor);
 		if (anchor_status == ar_device_anchor_query_status_success)
 		{
-			cp_drawable_set_device_anchor(FrameState.SwiftDrawable, FrameState.DeviceAnchor);
+//			if (OXRVisionOSCVars::CVarDisableReprojection.GetValueOnAnyThread())
+//			{
+//				// Note: Documentation says this will disable reprojection, but it does not apper to actually work:
+				// Log says: "Presenting a drawable without a device anchor. This drawable won't be presented."
+//				cp_drawable_set_device_anchor(FrameState.SwiftDrawable, nullptr);
+//			}
+//			else
+			{
+				cp_drawable_set_device_anchor(FrameState.SwiftDrawable, FrameState.DeviceAnchor);
+			}
+
 			// Get the HMD transform and cache it
 			FrameState.HeadTransform = ar_anchor_get_origin_from_anchor_transform(FrameState.DeviceAnchor);
+			UE_LOG(LogOXRVisionOS, VeryVerbose, TEXT("%s FC: %i SLF: 0x%x SD: 0x%x cp_drawable_set_device_anchor in FOXRVisionOSSession::StartSubmission with DeviceAnchor 0x%x  also getting head transform and buffering it for GT HeadTransform01=%f"), VOSThreadString(), FrameState.FrameCounter, FrameState.SwiftFrame, FrameState.SwiftDrawable, FrameState.DeviceAnchor, FrameState.HeadTransform.columns[0][1]);
+			
 		}
 		else
 		{
-			//TODO: this is also how you disable reprojection.  Hook that up to a cvar.
+			UE_LOG(LogOXRVisionOS, VeryVerbose, TEXT("%s FC: %i SLF: 0x%x SD: 0x%x cp_drawable_set_device_anchor in FOXRVisionOSSession::StartSubmission DeviceAnchor query failed, setting nullptr, no reprojection"), VOSThreadString(), FrameState.FrameCounter, FrameState.SwiftFrame, FrameState.SwiftDrawable);
+			
 			cp_drawable_set_device_anchor(FrameState.SwiftDrawable, nullptr);
 			// Leave FrameState.HeadTransform untouched, same as game thread. It may simply be identity.  But whatever it is we will go ahead and render with it.
 		}
+		RenderToGameHeadTransformWrite(FrameState);
 	}
-    
-    //cp_drawable_set_ar_pose(FrameState.SwiftDrawable, )  // I think this is may be a deprecated way of specifying the frame's transform.  Replaced with cp_drawable_set_device_anchor
 }
 
 // FOpenXRHMD::OnFinishRendering_RHIThread()
 XrResult FOXRVisionOSSession::XrEndFrame(
 	const XrFrameEndInfo* InFrameEndInfo)
 {
-	UE_LOG(LogOXRVisionOS, Verbose, TEXT("FOXRVisionOSSession::XrEndFrame"));
+	UE_LOG(LogOXRVisionOS, Verbose, TEXT("%s FC: %i SLF: 0x%x SD: 0x%x FOXRVisionOSSession::XrEndFrame"), VOSThreadString(), PipelinedFrameStateRHI.FrameCounter, PipelinedFrameStateRHI.SwiftFrame, PipelinedFrameStateRHI.SwiftDrawable);
 
 	SCOPED_NAMED_EVENT_TEXT("FOXRVisionOSSession::XrEndFrame", FColor::Turquoise);
 
@@ -1095,7 +1120,7 @@ XrResult FOXRVisionOSSession::XrEndFrame(
 
 	if (PipelinedFrameStateRHI.bSynchronizing)
 	{
-		//UE_LOG(LogOXRVisionOS, Log, TEXT("FOXRVisionOSSession XrEndFrame bSynchronizing=true fliparg=%i buffer=%i"), PipelinedFrameStateRHI.FrameCounter, PipelinedFrameStateRHI.BackBufferIndex);
+		UE_LOG(LogOXRVisionOS, Verbose, TEXT("FOXRVisionOSSession XrEndFrame bSynchronizing=true"));
 		return XrResult::XR_SUCCESS;
 	}
 
@@ -1123,6 +1148,8 @@ XrResult FOXRVisionOSSession::XrEndFrame(
         const FOXRVisionOSSwapchain& SwapchainData0 = *(reinterpret_cast<FOXRVisionOSSwapchain*>(View0.subImage.swapchain));
         const FOXRVisionOSSwapchain::FSwapchainImage& ReleasedImage = SwapchainData0.GetLastReleasedImage();
 		
+		UE_LOG(LogOXRVisionOS, VeryVerbose, TEXT("%s FC: %i SLF: 0x%x SD: 0x%x FOXRVisionOSSession::XrEndFrame SwapchainData0.GetLastReleasedImage() = 0x%x"), VOSThreadString(), PipelinedFrameStateRHI.FrameCounter, PipelinedFrameStateRHI.SwiftFrame, PipelinedFrameStateRHI.SwiftDrawable, ReleasedImage.Image.GetReference()->GetNativeResource());
+		
 		const XrCompositionLayerDepthInfoKHR* Depth0 = OpenXR::FindChainedStructByType<XrCompositionLayerDepthInfoKHR>(View0.next, (XrStructureType)XR_TYPE_COMPOSITION_LAYER_DEPTH_INFO_KHR);
 		const XrCompositionLayerDepthInfoKHR* Depth1 = OpenXR::FindChainedStructByType<XrCompositionLayerDepthInfoKHR>(View1.next, (XrStructureType)XR_TYPE_COMPOSITION_LAYER_DEPTH_INFO_KHR);
 		
@@ -1136,12 +1163,7 @@ XrResult FOXRVisionOSSession::XrEndFrame(
         MetalRHIVisionOS::PresentImmersive(Params);
     }
 
-	 	NextBackBufferIndex = (NextBackBufferIndex + 1) % OXRVisionOSSessionHelpers::OXRVisionOSBackbufferLength;
-	
-	// else
-	// {
-	// 	UE_LOG(LogOXRVisionOS, Log, TEXT("XrEndFrame no layers, discarding, FrameCounter: %d, BufferIndex: %d"), PipelinedFrameStateRHI.FrameCounter, NextBackBufferIndex);
-	// }
+	NextBackBufferIndex = (NextBackBufferIndex + 1) % OXRVisionOSSessionHelpers::OXRVisionOSBackbufferLength;
 
 	return XrResult::XR_SUCCESS;
 }
@@ -1173,16 +1195,16 @@ XrResult FOXRVisionOSSession::XrLocateViews(
         return XrResult::XR_ERROR_VIEW_CONFIGURATION_TYPE_UNSUPPORTED;
     }
     
-    FPipelinedFrameState& FrameState = GetPipelinedFrameStateForThread();
-    
-    FOXRVisionOSSpace& BaseSpace = *(FOXRVisionOSSpace*)ViewLocateInfo->space;
-    const int32 FrameCounter = FrameState.FrameCounter;
-    FTransform BaseSpaceInverseTransform;
-    XrSpaceLocationFlags BaseSpaceFlags;
-    FVector OutLinearVelocity; FVector OutAngularVelocity; XrSpaceVelocityFlags OutVelocityFlags; // These are unused here.
-    BaseSpace.GetInverseTransform(ViewLocateInfo->displayTime, BaseSpaceInverseTransform, BaseSpaceFlags, OutLinearVelocity, OutAngularVelocity, OutVelocityFlags);
+	const FOXRVisionOSSpace* BaseSpacePtr = ToFOXRVisionOSSpace(ViewLocateInfo->space);
+	if (BaseSpacePtr == nullptr) {
+		return XrResult::XR_ERROR_HANDLE_INVALID;
+	}
     
-    uint32_t NumViews = 0;
+	FPipelinedFrameState& FrameState = GetPipelinedFrameStateForThread();
+	
+	UE_LOG(LogOXRVisionOS, Verbose, TEXT("%s FC: %i SLF: 0x%x SD: 0x%x FOXRVisionOSSession::XrLocateViews"), VOSThreadString(), FrameState.FrameCounter, FrameState.SwiftFrame, FrameState.SwiftDrawable);
+
+	uint32_t NumViews = 0;
     if (IsInGameThread())
     {
         // We will use the NumViews from a previous render frame.
@@ -1208,29 +1230,31 @@ XrResult FOXRVisionOSSession::XrLocateViews(
             return XrResult::XR_ERROR_SIZE_INSUFFICIENT;
         }
         
-        //TODO make this more real???
-        ViewState->viewStateFlags = (XR_SPACE_LOCATION_POSITION_VALID_BIT | XR_SPACE_LOCATION_POSITION_TRACKED_BIT);
+        // The view is always located.
+        ViewState->viewStateFlags = (XR_VIEW_STATE_POSITION_VALID_BIT |
+									 XR_VIEW_STATE_POSITION_TRACKED_BIT |
+									 XR_VIEW_STATE_ORIENTATION_VALID_BIT |
+									 XR_VIEW_STATE_ORIENTATION_TRACKED_BIT);
         
         for (int Index = 0; Index < NumViews; ++Index)
         {
             if (IsInGameThread())
             {
                 const FLocateViewInfo& ViewInfoCache = GetLocateViewInfo_GameThread();
-                
-                // Adjust the camera transform for the current eye position.
-                simd_float4x4 CameraTransform = simd_mul(PipelinedFrameStateGame.HeadTransform, ViewInfoCache.ViewTransforms[Index]);
-                
-                Views[Index].pose = OXRVisionOS::ToXrPose(CameraTransform);
+                Views[Index].pose = OXRVisionOS::ToXrPose(ViewInfoCache.ViewTransforms[Index]);
                 Views[Index].fov = ViewInfoCache.HmdFovs[Index];
+				UE_LOG(LogOXRVisionOS, Verbose, TEXT("%s FC: %i SLF: 0x%x SD: 0x%x FOXRVisionOSSession::XrLocateViews cached HeadTransform01=%f xrposexy=%f,%f"), VOSThreadString(), FrameState.FrameCounter, FrameState.SwiftFrame, FrameState.SwiftDrawable, FrameState.HeadTransform.columns[0][1], Views[Index].pose.position.x, Views[Index].pose.position.y);
             }
             else
             {
                 check(IsInRenderingThread());
                 check(FrameState.SwiftDrawable);
                 
+				UE_LOG(LogOXRVisionOS, VeryVerbose, TEXT("%s FC: %i SLF: 0x%x SD: 0x%x cp_drawable_get_view in FOXRVisionOSSession::XrLocateViews view %i"), VOSThreadString(), FrameState.FrameCounter, FrameState.SwiftFrame, FrameState.SwiftDrawable, Index);
+
                 cp_view_t View = cp_drawable_get_view(FrameState.SwiftDrawable, Index);
                 simd_float4 Tangents = cp_view_get_tangents(View);
-                //simd_float2 DepthRange = cp_drawable_get_depth_range(FrameState.SwiftDrawable);  //TODO ??? do we need this?  There is an openxrextension for it, probably should use that.
+                //simd_float2 DepthRange = cp_drawable_get_depth_range(FrameState.SwiftDrawable);  //TODO ??? do we need this?  There is an openxrextension for it, could use that.
     //                                                                          DepthRange[1], /* nearZ */
     //                                                                          DepthRange[0], /* farZ */
     //                                                                          true); /* reverseZ */
@@ -1240,69 +1264,40 @@ XrResult FOXRVisionOSSession::XrLocateViews(
                 FrameState.HmdFovs[Index].angleUp       =  FMath::Atan(Tangents[2]);
                 FrameState.HmdFovs[Index].angleDown     = -FMath::Atan(Tangents[3]);
 
-                // Adjust the camera transform for the current eye position.
+                // Hack: We only xrLocateViews for the HMD device, and cp_view_get_transform gives us hmd relative transforms, so we don't need to do any additional math here.
+				// However OpenXR spec allows one to request the views in any space, which would require us to do some transforms and inverse transforms.
+				// Perhaps OpenXRHMD ought to be getting the view poses in view space rather than HMD device space because we could detect that here easily???
                 simd_float4x4 ViewTransform = cp_view_get_transform(View);
-                simd_float4x4 CameraTransform = simd_mul(PipelinedFrameStateRHI.HeadTransform, ViewTransform);
-                
-                Views[Index].pose = OXRVisionOS::ToXrPose(CameraTransform);
+                Views[Index].pose = OXRVisionOS::ToXrPose(ViewTransform);
                 Views[Index].fov = FrameState.HmdFovs[Index];
-                
+				
+				UE_LOG(LogOXRVisionOS, Verbose, TEXT("%s FC: %i SLF: 0x%x SD: 0x%x FOXRVisionOSSession::XrLocateViews cp_view_get_transformHeadTransform01=%f xrposexy=%f,%f"), VOSThreadString(), FrameState.FrameCounter, FrameState.SwiftFrame, FrameState.SwiftDrawable, FrameState.HeadTransform.columns[0][1], Views[Index].pose.position.x, Views[Index].pose.position.y);
+
                 // Cache the render thread data for the game thread
                 FLocateViewInfo& ViewInfoCache = GetLocateViewInfo_RenderThread();
                 ViewInfoCache.ViewTransforms[Index] = ViewTransform;
                 ViewInfoCache.HmdFovs[Index] = FrameState.HmdFovs[Index];
             }
         }
-
-		//StoreHMDPoseOffsetData(hmdInfo.devicePose);
 	}
 
-//	if (IsInGameThread())
-//	{
-//		ENQUEUE_RENDER_COMMAND(UpdateHMDPoseFromGT)([this, NewRenderingFrameState = PipelinedFrameStateGame](FRHICommandListImmediate& RHICmdList)
-//		{
-//			PipelinedFrameStateRendering.HMDPoseInTrackerSpace = NewRenderingFrameState.HMDPoseInTrackerSpace;
-//			RHICmdList.EnqueueLambda([this, NewRHIFrameState = PipelinedFrameStateRendering](FRHICommandListImmediate&)
-//			{
-//				PipelinedFrameStateRHI.HMDPoseInTrackerSpace = NewRHIFrameState.HMDPoseInTrackerSpace;
-//			});
-//		});
-//	}
-//	else
-//	{
-//		check(IsInRenderingThread());
-//		FRHICommandListImmediate& RHICmdList = GetImmediateCommandList_ForRenderCommand();
-//		RHICmdList.EnqueueLambda([this, NewRHIFrameState = PipelinedFrameStateRendering](FRHICommandListImmediate&)
-//		{
-//			PipelinedFrameStateRHI.HMDPoseInTrackerSpace = NewRHIFrameState.HMDPoseInTrackerSpace;
-//		});
-//	}
-
 	return XrResult::XR_SUCCESS;
 }
 
-// void FOXRVisionOSSession::StoreHMDPoseOffsetData(const PoseData& HmdPose)
-// {
-// 	FTransform HmdTransform = OXRVisionOS::ToFTransform(HmdPose);
-// 	// After a tracking disruption hmd transform often becomes slightly denormalized, more than our checkslow threshold.
-// 	HmdTransform.NormalizeRotation();
-
-// 	FPipelinedFrameState& FrameStateForThread = GetPipelinedFrameStateForThread();
-// 	FrameStateForThread.HMDPoseInTrackerSpace = HmdTransform;
-// }
-
 void FOXRVisionOSSession::GetHMDTransform(XrTime DisplayTime, FTransform& OutTransform, XrSpaceLocationFlags& OutLocationFlags)
 {
     //TODO just ignoring XrTime right now!!!
     
     FPipelinedFrameState& FrameState = GetPipelinedFrameStateForThread();
-    
+	
     OutLocationFlags = XR_SPACE_LOCATION_ORIENTATION_VALID_BIT
         | XR_SPACE_LOCATION_POSITION_VALID_BIT
         | XR_SPACE_LOCATION_ORIENTATION_TRACKED_BIT
     | XR_SPACE_LOCATION_POSITION_TRACKED_BIT;
     
     OutTransform = OXRVisionOS::ToFTransform(FrameState.HeadTransform);
+
+	UE_LOG(LogOXRVisionOS, VeryVerbose, TEXT("%s FC: %i SLF: 0x%x SD: 0x%x FOXRVisionOSSession::GetHMDTransform HeadTransform01=%f"), VOSThreadString(), FrameState.FrameCounter, FrameState.SwiftFrame, FrameState.SwiftDrawable, FrameState.HeadTransform.columns[0][1]);
 }
 
 void FOXRVisionOSSession::GetTransformForButton(EOXRVisionOSControllerButton Button, XrTime DisplayTime, FTransform& OutTransform, XrSpaceLocationFlags& OutLocationFlags, FVector& OutLinearVelocity, FVector& OutAngularVelocity, XrSpaceVelocityFlags& OutVelocityFlags)
@@ -1349,7 +1344,6 @@ void FOXRVisionOSSession::GetTransformForButton(EOXRVisionOSControllerButton But
 	// 		return;
 	// 	}
 	// 	OutLocationFlags = OXRVisionOS::ToXrSpaceLocationFlags(HMDResultData);
-	// 	OutVelocityFlags = 0; // gaze never provides velocity data
 
 	// 	// gazeDirection is in gazeOriginPose space.
 	// 	PoseData& HMDGazeOriginPose = HMDResultData.Hmd.GazeOriginPose;
@@ -1370,7 +1364,6 @@ void FOXRVisionOSSession::GetTransformForButton(EOXRVisionOSControllerButton But
 		}
 
 		//OutLocationFlags = OXRVisionOS::ToXrSpaceLocationFlags(ResultData);
-		//OutVelocityFlags = OXRVisionOS::ToXrSpaceVelocityFlags(ResultData);
 
 		if ((OutLocationFlags & (XR_SPACE_LOCATION_POSITION_VALID_BIT | XR_SPACE_LOCATION_ORIENTATION_VALID_BIT)))
 		{
@@ -1393,15 +1386,6 @@ void FOXRVisionOSSession::GetTransformForButton(EOXRVisionOSControllerButton But
 				OutTransform.SetRotation(OutTransform.TransformRotation(AimPoseAdjustment));
 			}
 		}
-
-		// if (OutVelocityFlags & XR_SPACE_VELOCITY_LINEAR_VALID_BIT)
-		// {
-		// 	OutLinearVelocity = OXRVisionOS::ToFVector(ResultData.velocity);
-		// }
-		// if (OutVelocityFlags & XR_SPACE_VELOCITY_ANGULAR_VALID_BIT)
-		// {
-		// 	OutAngularVelocity = OXRVisionOS::ToFVector(ResultData.angularVelocity);
-		// }
 	}
 }
 
@@ -1458,3 +1442,15 @@ FOXRVisionOSSession::FPipelinedFrameState& FOXRVisionOSSession::GetPipelinedFram
 {
 	return const_cast<FOXRVisionOSSession::FPipelinedFrameState&>(static_cast<const FOXRVisionOSSession&>(*this).GetPipelinedFrameStateForThread());
 }
+
+// Read from the current index.
+simd_float4x4 FOXRVisionOSSession::RenderToGameHeadTransformRead(const FPipelinedFrameState& FrameState)
+{
+	return RenderToGameHeadTransform[FrameState.RenderToGameHeadTransformIndexRead % RenderToGameHeadTransformBufferLength];
+}
+
+// Write to the next index that the game thread will read from.
+void FOXRVisionOSSession::RenderToGameHeadTransformWrite(FPipelinedFrameState& FrameState)
+{
+	RenderToGameHeadTransform[(FrameState.RenderToGameHeadTransformIndexRead + 1) % RenderToGameHeadTransformBufferLength] = FrameState.HeadTransform;
+}
diff --git a/Engine/Platforms/VisionOS/Plugins/Runtime/OpenXRVisionOS/Source/OXRVisionOS/Private/OXRVisionOSSession.h b/Engine/Platforms/VisionOS/Plugins/Runtime/OpenXRVisionOS/Source/OXRVisionOS/Private/OXRVisionOSSession.h
--- a/Engine/Platforms/VisionOS/Plugins/Runtime/OpenXRVisionOS/Source/OXRVisionOS/Private/OXRVisionOSSession.h
+++ b/Engine/Platforms/VisionOS/Plugins/Runtime/OpenXRVisionOS/Source/OXRVisionOS/Private/OXRVisionOSSession.h
@@ -50,6 +50,8 @@ public:
 	XrResult XrCreateActionSpace(
 		const XrActionSpaceCreateInfo* createInfo,
 		XrSpace* space);
+	
+	const FOXRVisionOSSpace* ToFOXRVisionOSSpace(const XrSpace& Space);
 
 	XrResult DestroySpace(
 		FOXRVisionOSSpace* Space);
@@ -160,7 +162,12 @@ public:
 		bool bSynchronizing = true; // when true we are in the process of starting up a new session and the new frame state has not yet propagated.
 		XrFovf HmdFovs[2];
         int32 LocateViewInfoBufferIndex = 0;
+		int32 RenderToGameHeadTransformIndexRead = -1;
 	};
+	
+	simd_float4x4 RenderToGameHeadTransformRead(const FPipelinedFrameState& FrameState);
+	void RenderToGameHeadTransformWrite(FPipelinedFrameState& FrameState);
+
 	const FPipelinedFrameState& GetPipelinedFrameStateForThread() const;
 	FPipelinedFrameState& GetPipelinedFrameStateForThread();
 	void GetHMDTransform(XrTime DisplayTime, FTransform& OutTransform, XrSpaceLocationFlags& OutLocationFlags);
@@ -180,7 +187,7 @@ public:
     void LocateViewInfoIndexAdvance() { ++PipelinedFrameStateGame.LocateViewInfoBufferIndex; }
     const FLocateViewInfo& GetLocateViewInfo_GameThread() { return LocateViewInfoBuffer[(PipelinedFrameStateGame.LocateViewInfoBufferIndex - 1) % LocateViewInfoBufferLength]; }
     FLocateViewInfo& GetLocateViewInfo_RenderThread() { return LocateViewInfoBuffer[PipelinedFrameStateRendering.LocateViewInfoBufferIndex % LocateViewInfoBufferLength]; }
-
+	
 private:
 	bool bCreateFailed = false;
 	bool bIsRunning = false;
@@ -198,18 +205,12 @@ private:
 	FPlatformMemory::FPlatformVirtualMemoryBlock m_reprojectionData;
 	FPlatformMemory::FPlatformVirtualMemoryBlock m_displayData;
 
-
 	TArray<TSharedPtr<FOXRVisionOSSwapchain, ESPMode::ThreadSafe>> Swapchains;
 
 	int32 NextBackBufferIndex = 0;
 	int32 PreviousBackBufferIndexInBegin = -1;
 	int32 SessionFrameCounter = 0;
 
-	// Because OpenXR works with view poses, but works with hmd pose we cache an hmd pose when we get the trackign data.
-	// If the OpenXR Application modifies the view poses such that they are not relative to the hmd pose this will fail and the altered view poses will not behave as expected.
-	// This should work correctly if api modifies the interpupillary distance, provided that is properly pipelined on the api side (meaning it is constant for each fliparg).
-	//void StoreHMDPoseOffsetData(const PoseData& HmdPose);
-
 	static const int32 INVALID_DEVICE_HANDLE = -1;
 	int32 HMDHandle = INVALID_DEVICE_HANDLE;
 
@@ -230,6 +231,8 @@ private:
 	FPipelinedFrameState	PipelinedFrameStateRHI;
 	int32 CachedBeginFlipFrameCounter = INT32_MAX;
 	
+	static const int32 RenderToGameHeadTransformBufferLength = 3;
+	simd_float4x4 RenderToGameHeadTransform[RenderToGameHeadTransformBufferLength];
 	//TEMP see usages
 	FCriticalSection CriticalSection_ar_world_tracking_provider_query_device_anchor_at_timestamp;
 
diff --git a/Engine/Platforms/VisionOS/Source/Runtime/Launch/Source/UESwift.swift b/Engine/Platforms/VisionOS/Source/Runtime/Launch/Source/UESwift.swift
--- a/Engine/Platforms/VisionOS/Source/Runtime/Launch/Source/UESwift.swift
+++ b/Engine/Platforms/VisionOS/Source/Runtime/Launch/Source/UESwift.swift
@@ -75,6 +75,8 @@ struct UEContentConfiguration: CompositorLayerConfiguration {
 		//configuration.depthFormat = .depth32Float  			//PF_R32_FLOAT   			// This is correct for mobile forward
 		configuration.depthFormat = .depth32Float_stencil8 		//PF_DepthStencil   // This is correct for deferred
 //PFSWITCH
+		
+		configuration.defaultDepthRange = [Float.greatestFiniteMagnitude, 0.1]
 	}
 }
 
