UE-208429 VisionOS: hand tracking

diff --git a/Engine/Platforms/VisionOS/Plugins/Runtime/OpenXRHandTracking/OpenXRHandTracking_VisionOS.uplugin b/Engine/Platforms/VisionOS/Plugins/Runtime/OpenXRHandTracking/OpenXRHandTracking_VisionOS.uplugin
new file mode 100644
--- /dev/null
+++ b/Engine/Platforms/VisionOS/Plugins/Runtime/OpenXRHandTracking/OpenXRHandTracking_VisionOS.uplugin
@@ -0,0 +1,17 @@
+{
+	"FileVersion": 3,
+	"bIsPluginExtension": true,
+	"SupportedTargetPlatforms": [
+		"VisionOS"
+	],
+	"Modules": [
+		{
+			"Name": "OpenXRHandTracking",
+			"Type": "Runtime",
+			"PlatformAllowList": [
+				"VisionOS"
+			]
+		}
+	]
+}
+
diff --git a/Engine/Platforms/VisionOS/Plugins/Runtime/OpenXRVisionOS/Source/OXRVisionOS/Private/OXRVisionOS.cpp b/Engine/Platforms/VisionOS/Plugins/Runtime/OpenXRVisionOS/Source/OXRVisionOS/Private/OXRVisionOS.cpp
--- a/Engine/Platforms/VisionOS/Plugins/Runtime/OpenXRVisionOS/Source/OXRVisionOS/Private/OXRVisionOS.cpp
+++ b/Engine/Platforms/VisionOS/Plugins/Runtime/OpenXRVisionOS/Source/OXRVisionOS/Private/OXRVisionOS.cpp
@@ -748,15 +748,53 @@ namespace OXRVisionOS
 			hapticActionInfo);
 	}
 
+// XR_EXT_hand_tracking begin
+	XRAPI_ATTR XrResult XRAPI_CALL xrCreateHandTrackerEXT(
+		XrSession                                   session,
+		const XrHandTrackerCreateInfoEXT*           createInfo,
+		XrHandTrackerEXT*                           handTracker)
+	{
+		OXRVISIONOS_CHECK_SESSION_EARLY_RETURN(session);
+		
+		FOXRVisionOSSession* Session = (FOXRVisionOSSession*)session;
+		return Session->XrCreateHandTrackerEXT(
+			createInfo,
+			handTracker);
+	}
+
+	XRAPI_ATTR XrResult XRAPI_CALL xrDestroyHandTrackerEXT(
+		XrHandTrackerEXT                            handTracker)
+	{
+		FOXRVisionOSSession::FOXRVisionOSHandTracker* HandTracker = (FOXRVisionOSSession::FOXRVisionOSHandTracker*)handTracker;
+		return HandTracker->XrDestroyHandTrackerEXT();
+	}
+
+	XRAPI_ATTR XrResult XRAPI_CALL xrLocateHandJointsEXT(
+		XrHandTrackerEXT                            handTracker,
+		const XrHandJointsLocateInfoEXT*            locateInfo,
+		XrHandJointLocationsEXT*                    locations)
+	{
+		FOXRVisionOSSession::FOXRVisionOSHandTracker* HandTracker = (FOXRVisionOSSession::FOXRVisionOSHandTracker*)handTracker;
+		return HandTracker->XrLocateHandJointsEXT(
+			locateInfo,
+			locations);
+	}
+// XR_EXT_hand_tracking end
+
 // Any supported extension functions would need to be mapped as well.  That could go here.
 
 #define OXRVISIONOS_MAP_GLOBAL_FUNCTION(Type,Func) OXRVisionOS::GlobalFunctionMap.Add(FString(TEXT(#Func)), (PFN_xrVoidFunction)(OXRVisionOS::Func));
 #define OXRVISIONOS_MAP_INSTANCE_FUNCTION(Type,Func) OXRVisionOS::InstanceFunctionMap.Add(FString(TEXT(#Func)), (PFN_xrVoidFunction)(OXRVisionOS::Func));
+#define OXRVISIONOS_MAP_INSTANCE_FUNCTION_EXT(Func) OXRVisionOS::InstanceFunctionMap.Add(FString(TEXT(#Func)), (PFN_xrVoidFunction)(OXRVisionOS::Func));
 
 	void BuildFunctionMaps()
 	{
 		ENUM_XR_ENTRYPOINTS_GLOBAL(OXRVISIONOS_MAP_GLOBAL_FUNCTION);
 		ENUM_XR_ENTRYPOINTS(OXRVISIONOS_MAP_INSTANCE_FUNCTION);
+		
+		OXRVISIONOS_MAP_INSTANCE_FUNCTION_EXT(xrCreateHandTrackerEXT);
+		OXRVISIONOS_MAP_INSTANCE_FUNCTION_EXT(xrDestroyHandTrackerEXT);
+		OXRVISIONOS_MAP_INSTANCE_FUNCTION_EXT(xrLocateHandJointsEXT);
 	}
 
 #undef OXRVISIONOS_MAP_INSTANCE_FUNCTION
@@ -781,7 +819,9 @@ void FOXRVisionOS::StartupModule()
 
 	SupportedExtensions.Add(XR_KHR_COMPOSITION_LAYER_DEPTH_EXTENSION_NAME, XR_KHR_composition_layer_depth_SPEC_VERSION);
 	SupportedExtensions.Add(XR_EPIC_OXRVISIONOS_CONTROLLER_NAME, XR_EPIC_oxrvisionos_controller_SPEC_VERSION);
-	//	SupportedExtensions.Add(XR_EXT_EYE_GAZE_INTERACTION_EXTENSION_NAME, XR_EXT_eye_gaze_interaction_SPEC_VERSION);
+	SupportedExtensions.Add(XR_EXT_HAND_TRACKING_EXTENSION_NAME, XR_EXT_hand_tracking);
+	//SupportedExtensions.Add(XR_EXT__EXTENSION_NAME, XR_MSFT_hand_interaction);
+	//SupportedExtensions.Add(XR_EXT__EXTENSION_NAME, XR_EXT_eye_gaze_interaction_SPEC_VERSION);
 }
 
 void FOXRVisionOS::ShutdownModule()
diff --git a/Engine/Platforms/VisionOS/Plugins/Runtime/OpenXRVisionOS/Source/OXRVisionOS/Private/OXRVisionOSInstance.cpp b/Engine/Platforms/VisionOS/Plugins/Runtime/OpenXRVisionOS/Source/OXRVisionOS/Private/OXRVisionOSInstance.cpp
--- a/Engine/Platforms/VisionOS/Plugins/Runtime/OpenXRVisionOS/Source/OXRVisionOS/Private/OXRVisionOSInstance.cpp
+++ b/Engine/Platforms/VisionOS/Plugins/Runtime/OpenXRVisionOS/Source/OXRVisionOS/Private/OXRVisionOSInstance.cpp
@@ -237,8 +237,17 @@ XrResult FOXRVisionOSInstance::XrGetSystemProperties(
 	{
 		return XrResult::XR_ERROR_VALIDATION_FAILURE;
 	}
-
+	
+	// Cache the next ptr, so we don't overwrite it with all the other system properties.
+	void* NextTmp = Properties->next;
 	*Properties = SystemProperties;
+	Properties->next = NextTmp;
+	
+	XrSystemHandTrackingPropertiesEXT* HandTrackingProperties = OpenXR::FindChainedStructByType<XrSystemHandTrackingPropertiesEXT>(Properties, XR_TYPE_SYSTEM_HAND_TRACKING_PROPERTIES_EXT);
+	if (HandTrackingProperties) {
+		HandTrackingProperties->supportsHandTracking = XR_TRUE;
+	}
+
 	return XrResult::XR_SUCCESS;
 }
 
diff --git a/Engine/Platforms/VisionOS/Plugins/Runtime/OpenXRVisionOS/Source/OXRVisionOS/Private/OXRVisionOSSession.cpp b/Engine/Platforms/VisionOS/Plugins/Runtime/OpenXRVisionOS/Source/OXRVisionOS/Private/OXRVisionOSSession.cpp
--- a/Engine/Platforms/VisionOS/Plugins/Runtime/OpenXRVisionOS/Source/OXRVisionOS/Private/OXRVisionOSSession.cpp
+++ b/Engine/Platforms/VisionOS/Plugins/Runtime/OpenXRVisionOS/Source/OXRVisionOS/Private/OXRVisionOSSession.cpp
@@ -596,6 +596,8 @@ XrResult FOXRVisionOSSession::XrSyncActions(
 //	{
 //		ActionSet->SyncActions(Controllers, Tracker);
 //	}
+	
+	SyncHandTracking();
 
 	return XR_SUCCESS;
 }
@@ -693,11 +695,18 @@ XrResult FOXRVisionOSSession::XrBeginSession(
 	// Create World Tracking data provider
 	ar_world_tracking_configuration_t ArKitWorldTrackingConfiguration = ar_world_tracking_configuration_create();
 	ARKitWorldTrackingProvider = ar_world_tracking_provider_create(ArKitWorldTrackingConfiguration);
+	
+	// Create Hand Tracking data provider
+	BeginHandTracking();
+	ar_hand_tracking_configuration_t ArKitHandTrackingConfiguration = ar_hand_tracking_configuration_create();
+	ARKitHandTrackingProvider = ar_hand_tracking_provider_create(ArKitHandTrackingConfiguration);
+	
 	// Create Data providers
 	ar_data_providers_t ARKitDataProviders = ar_data_providers_create();
 	// Add
 	ar_data_providers_add_data_provider(ARKitDataProviders, ARKitWorldTrackingProvider);
-	
+	ar_data_providers_add_data_provider(ARKitDataProviders, ARKitHandTrackingProvider);
+
 	// Start the arkit session
 	ar_session_run(ARKitSession, ARKitDataProviders);
 
@@ -1454,3 +1463,171 @@ void FOXRVisionOSSession::RenderToGameHeadTransformWrite(FPipelinedFrameState& F
 {
 	RenderToGameHeadTransform[(FrameState.RenderToGameHeadTransformIndexRead + 1) % RenderToGameHeadTransformBufferLength] = FrameState.HeadTransform;
 }
+
+
+XrResult FOXRVisionOSSession::XrCreateHandTrackerEXT(
+	const XrHandTrackerCreateInfoEXT*           createInfo,
+	XrHandTrackerEXT*                           handTracker)
+{
+	check(createInfo->handJointSet == XR_HAND_JOINT_SET_DEFAULT_EXT);  // We only support the default.
+	FOXRVisionOSHandTracker& NewHandTracker = GetHandTracker(createInfo->hand);
+	NewHandTracker.bCreated = true;
+	*handTracker = (XrHandTrackerEXT)&NewHandTracker;
+	return XR_SUCCESS;
+}
+
+XrResult FOXRVisionOSSession::FOXRVisionOSHandTracker::XrDestroyHandTrackerEXT()
+{
+	bCreated = false;
+	return XR_SUCCESS;
+}
+
+XrResult FOXRVisionOSSession::FOXRVisionOSHandTracker::XrLocateHandJointsEXT(
+	const XrHandJointsLocateInfoEXT*            locateInfo,
+	XrHandJointLocationsEXT*                    locations)
+{
+	//locateInfo->time // TODO: just ignoring time for now, we will always get the latest.
+	//locateInfo->baseSpace  //TODO: ignoring this, always get in the tracking space.
+	
+	locations->isActive = bIsActive;
+	locations->jointCount = XR_HAND_JOINT_COUNT_EXT;
+	for (int i = 0; i < locations->jointCount; ++i) 
+	{
+		locations->jointLocations[i] = JointLocations[i];
+	}
+	
+	XrPosef pose = locations->jointLocations[XR_HAND_JOINT_THUMB_TIP_EXT].pose;
+	UE_LOG(LogTemp, Log, TEXT("XrLocateHandJointsEXT() locating thumb tip flags: %i  xyz %0.2f, %0.2f, %0.2f"), locations->jointLocations[XR_HAND_JOINT_THUMB_TIP_EXT].locationFlags, pose.position.x, pose.position.y, pose.position.z);
+
+	return XR_SUCCESS;
+}
+
+void FOXRVisionOSSession::BeginHandTracking()
+{
+	if (JointMap.IsEmpty()) 
+	{
+		JointMap.Reserve(XR_HAND_JOINT_COUNT_EXT);
+		//JointMap.Add(,	XR_HAND_JOINT_PALM_EXT); // ARKit has no palm, we will synthesize one.
+		JointMap.Add(TPair<ar_hand_skeleton_joint_name_t,XrHandJointEXT>(ar_hand_skeleton_joint_name_wrist,	XR_HAND_JOINT_WRIST_EXT));
+		JointMap.Add(TPair<ar_hand_skeleton_joint_name_t,XrHandJointEXT>(ar_hand_skeleton_joint_name_thumb_knuckle,	XR_HAND_JOINT_THUMB_METACARPAL_EXT));
+		JointMap.Add(TPair<ar_hand_skeleton_joint_name_t,XrHandJointEXT>(ar_hand_skeleton_joint_name_thumb_intermediate_base,	XR_HAND_JOINT_THUMB_PROXIMAL_EXT));
+		JointMap.Add(TPair<ar_hand_skeleton_joint_name_t,XrHandJointEXT>(ar_hand_skeleton_joint_name_thumb_intermediate_tip,	XR_HAND_JOINT_THUMB_DISTAL_EXT));
+		JointMap.Add(TPair<ar_hand_skeleton_joint_name_t,XrHandJointEXT>(ar_hand_skeleton_joint_name_thumb_tip,	XR_HAND_JOINT_THUMB_TIP_EXT));
+		JointMap.Add(TPair<ar_hand_skeleton_joint_name_t,XrHandJointEXT>(ar_hand_skeleton_joint_name_index_finger_metacarpal,	XR_HAND_JOINT_INDEX_METACARPAL_EXT));
+		JointMap.Add(TPair<ar_hand_skeleton_joint_name_t,XrHandJointEXT>(ar_hand_skeleton_joint_name_index_finger_knuckle,	XR_HAND_JOINT_INDEX_PROXIMAL_EXT));
+		JointMap.Add(TPair<ar_hand_skeleton_joint_name_t,XrHandJointEXT>(ar_hand_skeleton_joint_name_index_finger_intermediate_base,	XR_HAND_JOINT_INDEX_INTERMEDIATE_EXT));
+		JointMap.Add(TPair<ar_hand_skeleton_joint_name_t,XrHandJointEXT>(ar_hand_skeleton_joint_name_index_finger_intermediate_tip,	XR_HAND_JOINT_INDEX_DISTAL_EXT));
+		JointMap.Add(TPair<ar_hand_skeleton_joint_name_t,XrHandJointEXT>(ar_hand_skeleton_joint_name_index_finger_tip,	XR_HAND_JOINT_INDEX_TIP_EXT));
+		JointMap.Add(TPair<ar_hand_skeleton_joint_name_t,XrHandJointEXT>(ar_hand_skeleton_joint_name_middle_finger_metacarpal,	XR_HAND_JOINT_MIDDLE_METACARPAL_EXT));
+		JointMap.Add(TPair<ar_hand_skeleton_joint_name_t,XrHandJointEXT>(ar_hand_skeleton_joint_name_middle_finger_knuckle,	XR_HAND_JOINT_MIDDLE_PROXIMAL_EXT));
+		JointMap.Add(TPair<ar_hand_skeleton_joint_name_t,XrHandJointEXT>(ar_hand_skeleton_joint_name_middle_finger_intermediate_base,	XR_HAND_JOINT_MIDDLE_INTERMEDIATE_EXT));
+		JointMap.Add(TPair<ar_hand_skeleton_joint_name_t,XrHandJointEXT>(ar_hand_skeleton_joint_name_middle_finger_intermediate_tip,	XR_HAND_JOINT_MIDDLE_DISTAL_EXT));
+		JointMap.Add(TPair<ar_hand_skeleton_joint_name_t,XrHandJointEXT>(ar_hand_skeleton_joint_name_middle_finger_tip,	XR_HAND_JOINT_MIDDLE_TIP_EXT));
+		JointMap.Add(TPair<ar_hand_skeleton_joint_name_t,XrHandJointEXT>(ar_hand_skeleton_joint_name_ring_finger_metacarpal,	XR_HAND_JOINT_RING_METACARPAL_EXT));
+		JointMap.Add(TPair<ar_hand_skeleton_joint_name_t,XrHandJointEXT>(ar_hand_skeleton_joint_name_ring_finger_knuckle,	XR_HAND_JOINT_RING_PROXIMAL_EXT));
+		JointMap.Add(TPair<ar_hand_skeleton_joint_name_t,XrHandJointEXT>(ar_hand_skeleton_joint_name_ring_finger_intermediate_base,	XR_HAND_JOINT_RING_INTERMEDIATE_EXT));
+		JointMap.Add(TPair<ar_hand_skeleton_joint_name_t,XrHandJointEXT>(ar_hand_skeleton_joint_name_ring_finger_intermediate_tip,	XR_HAND_JOINT_RING_DISTAL_EXT));
+		JointMap.Add(TPair<ar_hand_skeleton_joint_name_t,XrHandJointEXT>(ar_hand_skeleton_joint_name_ring_finger_tip,	XR_HAND_JOINT_RING_TIP_EXT));
+		JointMap.Add(TPair<ar_hand_skeleton_joint_name_t,XrHandJointEXT>(ar_hand_skeleton_joint_name_little_finger_metacarpal,				XR_HAND_JOINT_LITTLE_METACARPAL_EXT));
+		JointMap.Add(TPair<ar_hand_skeleton_joint_name_t,XrHandJointEXT>(ar_hand_skeleton_joint_name_little_finger_knuckle,					XR_HAND_JOINT_LITTLE_PROXIMAL_EXT));
+		JointMap.Add(TPair<ar_hand_skeleton_joint_name_t,XrHandJointEXT>(ar_hand_skeleton_joint_name_little_finger_intermediate_base,		XR_HAND_JOINT_LITTLE_INTERMEDIATE_EXT));
+		JointMap.Add(TPair<ar_hand_skeleton_joint_name_t,XrHandJointEXT>(ar_hand_skeleton_joint_name_little_finger_intermediate_tip,		XR_HAND_JOINT_LITTLE_DISTAL_EXT));
+		JointMap.Add(TPair<ar_hand_skeleton_joint_name_t,XrHandJointEXT>(ar_hand_skeleton_joint_name_little_finger_tip,						XR_HAND_JOINT_LITTLE_TIP_EXT));
+		//ar_hand_skeleton_joint_name_forearm_arm  // OpenXR does not support these forearm joints, there is an extension to support an elbow, we could add support for that.
+		//ar_hand_skeleton_joint_name_forearm_wrist
+	}
+	
+	FOXRVisionOSHandTracker& LeftHandTracker = GetHandTracker(XR_HAND_LEFT_EXT);
+	FOXRVisionOSHandTracker& RightHandTracker = GetHandTracker(XR_HAND_RIGHT_EXT);
+	LeftHandTracker.Anchor = ar_hand_anchor_create();
+	RightHandTracker.Anchor = ar_hand_anchor_create();
+}
+
+void FOXRVisionOSSession::SyncHandTracking()
+{
+	// We will retain the joint local transforms and the hand transforms so that partial updates can be applied.
+	
+	FOXRVisionOSHandTracker& LeftHandTracker = GetHandTracker(XR_HAND_LEFT_EXT);
+	FOXRVisionOSHandTracker& RightHandTracker = GetHandTracker(XR_HAND_RIGHT_EXT);
+	
+	bool Success = ar_hand_tracking_provider_get_latest_anchors(ARKitHandTrackingProvider, LeftHandTracker.Anchor, RightHandTracker.Anchor);
+	if (Success == false) 
+	{
+		return;
+	}
+
+	auto SyncHand = []( FOXRVisionOSHandTracker& HandTracker, const TArray<TPair<ar_hand_skeleton_joint_name_t,XrHandJointEXT>>& TheJointMap, bool bIsLeft)
+	{
+		HandTracker.bIsActive = true;
+		
+		bool bHandTracked = ar_trackable_anchor_is_tracked(HandTracker.Anchor);
+		if (bHandTracked) 
+		{
+			HandTracker.HandTransform = ar_anchor_get_origin_from_anchor_transform(HandTracker.Anchor);
+		}
+		
+		// The left hand's joints are all flipped vs the openxr spec, we need to rotate them 180 degrees around z.
+		const simd_float4x4 Transform180Z = simd_diagonal_matrix(simd_make_float4(-1,-1,1,1));
+
+		ar_hand_skeleton_t Skeleton = ar_hand_anchor_get_hand_skeleton(HandTracker.Anchor);
+		for(auto& Pair : TheJointMap)
+		{
+			ar_skeleton_joint_t Joint = ar_hand_skeleton_get_joint_named(Skeleton, Pair.Key);
+			simd_float4x4& JointLocalTransform = HandTracker.JointLocalTransforms[Pair.Value];
+			XrHandJointLocationEXT& JointLocation = HandTracker.JointLocations[Pair.Value];
+			
+			bool bIsTracked = ar_skeleton_joint_is_tracked(Joint);
+			if (bIsTracked) 
+			{
+				JointLocation.locationFlags = XR_SPACE_LOCATION_POSITION_VALID_BIT | XR_SPACE_LOCATION_ORIENTATION_VALID_BIT | XR_SPACE_LOCATION_POSITION_TRACKED_BIT | XR_SPACE_LOCATION_ORIENTATION_TRACKED_BIT;
+			}
+			else 
+			{
+				// Clear the tracked bits, but leave the valid bits and transforms as they are.  This means the joint is valid when it has been tracked once.
+				JointLocation.locationFlags = JointLocation.locationFlags && (XR_SPACE_LOCATION_POSITION_VALID_BIT | XR_SPACE_LOCATION_ORIENTATION_VALID_BIT);
+			}
+			JointLocalTransform = ar_skeleton_joint_get_anchor_from_joint_transform(Joint); // Even untracked joints do get updates from their parent's positions.
+
+			// We udpate the tracking space transform regardless of tracking status.  If either the local transform or the hand transform have updated it is useful.  If neither it doesn't really hurt.
+			simd_float4x4 JointTrackingTransform = matrix_multiply(HandTracker.HandTransform, JointLocalTransform);
+			if (bIsLeft) {
+				JointTrackingTransform = matrix_multiply(JointTrackingTransform, Transform180Z);
+			}
+			JointLocation.pose = OXRVisionOS::ToXrPose(JointTrackingTransform);
+			JointLocation.radius = 0.005f;  // In Meters. We could estimate this from the distance between joints and some allometric data about humans...
+		}
+		
+		// Special case for the palm, which ARKit does not provide.
+		{
+			//We will take the Middle Finger Metacarpal as the rotation of the Palm.
+			//We will average the Middle Finger Proximal and Middle Finger Metacarpal as the position of the Palm.
+			//The palm's flags will be the & of the two other joints flags.
+			//If one of the other bones is tracked and the other is not we might get an odd transform for the palm, but it will not be flagged as tracked.
+			simd_float4x4& PalmJointLocalTransform = HandTracker.JointLocalTransforms[XR_HAND_JOINT_PALM_EXT];
+			const simd_float4x4& MiddleProximalJointLocalTransform = HandTracker.JointLocalTransforms[XR_HAND_JOINT_MIDDLE_PROXIMAL_EXT];
+			const simd_float4x4& MiddleMetacarpalJointLocalTransform = HandTracker.JointLocalTransforms[XR_HAND_JOINT_MIDDLE_METACARPAL_EXT];
+			PalmJointLocalTransform = MiddleMetacarpalJointLocalTransform;
+			PalmJointLocalTransform.columns[3][0] = (MiddleProximalJointLocalTransform.columns[3][0] + MiddleMetacarpalJointLocalTransform.columns[3][0]) * 0.5f;
+			PalmJointLocalTransform.columns[3][1] = (MiddleProximalJointLocalTransform.columns[3][1] + MiddleMetacarpalJointLocalTransform.columns[3][1]) * 0.5f;
+			PalmJointLocalTransform.columns[3][2] = (MiddleProximalJointLocalTransform.columns[3][2] + MiddleMetacarpalJointLocalTransform.columns[3][2]) * 0.5f;
+			
+			simd_float4x4 JointTrackingTransform = matrix_multiply(HandTracker.HandTransform, PalmJointLocalTransform);
+			if (bIsLeft) {
+				JointTrackingTransform = matrix_multiply(JointTrackingTransform, Transform180Z);
+			}
+			
+			XrHandJointLocationEXT& PalmJoint = HandTracker.JointLocations[XR_HAND_JOINT_PALM_EXT];
+			XrHandJointLocationEXT& MiddleProximalJoint = HandTracker.JointLocations[XR_HAND_JOINT_MIDDLE_PROXIMAL_EXT];
+			XrHandJointLocationEXT& MiddleMetacarpalJoint = HandTracker.JointLocations[XR_HAND_JOINT_MIDDLE_METACARPAL_EXT];
+			PalmJoint.locationFlags = MiddleProximalJoint.locationFlags & MiddleMetacarpalJoint.locationFlags;
+			PalmJoint.pose = OXRVisionOS::ToXrPose(JointTrackingTransform);
+			PalmJoint.radius = 0.005f;  // In Meters. We could estimate this from the distance between joints and some allometric data about humans...
+
+		}
+	};
+	
+	SyncHand(LeftHandTracker, JointMap, true);
+	SyncHand(RightHandTracker, JointMap, false);
+}
+
+
diff --git a/Engine/Platforms/VisionOS/Plugins/Runtime/OpenXRVisionOS/Source/OXRVisionOS/Private/OXRVisionOSSession.h b/Engine/Platforms/VisionOS/Plugins/Runtime/OpenXRVisionOS/Source/OXRVisionOS/Private/OXRVisionOSSession.h
--- a/Engine/Platforms/VisionOS/Plugins/Runtime/OpenXRVisionOS/Source/OXRVisionOS/Private/OXRVisionOSSession.h
+++ b/Engine/Platforms/VisionOS/Plugins/Runtime/OpenXRVisionOS/Source/OXRVisionOS/Private/OXRVisionOSSession.h
@@ -136,6 +136,10 @@ public:
 		uint32_t viewCapacityInput,
 		uint32_t* viewCountOutput,
 		XrView* views);
+	
+	XrResult XrCreateHandTrackerEXT(
+		const XrHandTrackerCreateInfoEXT*           createInfo,
+		XrHandTrackerEXT*                           handTracker);
 
 	XrTime GetCurrentTime() const;
 
@@ -242,4 +246,41 @@ private:
 	ar_session_t ARKitSession = nullptr;
 	ar_world_tracking_provider_t ARKitWorldTrackingProvider = nullptr;
 	ar_device_anchor_t ARKitHMDAnchor = nullptr;
+	
+	// Hand Tracking
+public:
+	struct FOXRVisionOSHandTracker
+	{
+		FOXRVisionOSHandTracker()
+		{
+			HandTransform = matrix_identity_float4x4;
+			for (int i = 0; i < XR_HAND_JOINT_COUNT_EXT; ++i)
+			{
+				JointLocations[i].pose.orientation.w = 1;
+				JointLocalTransforms[i] = matrix_identity_float4x4;
+			}
+		}
+		XrResult XrDestroyHandTrackerEXT();
+
+		XrResult XrLocateHandJointsEXT(
+			const XrHandJointsLocateInfoEXT*            locateInfo,
+			XrHandJointLocationsEXT*                    locations);
+		
+		bool bCreated = false;
+		bool bIsActive = false;
+		ar_hand_anchor_t Anchor;
+		simd_float4x4 HandTransform;
+		simd_float4x4 JointLocalTransforms[XR_HAND_JOINT_COUNT_EXT];
+		XrHandJointLocationEXT JointLocations[XR_HAND_JOINT_COUNT_EXT] = {};
+	};
+private:
+	void BeginHandTracking();
+	ar_hand_tracking_provider_t ARKitHandTrackingProvider = nullptr;
+	TArray<TPair<ar_hand_skeleton_joint_name_t,XrHandJointEXT>> JointMap;
+	FOXRVisionOSHandTracker HandTrackers[2];
+	static_assert(XR_HAND_LEFT_EXT == 1 && XR_HAND_RIGHT_EXT == 2);  // We will assume this to use value-1 as array index.
+	FOXRVisionOSHandTracker& GetHandTracker(XrHandEXT Hand) { return HandTrackers[Hand - 1]; }
+	const FOXRVisionOSHandTracker& GetHandTracker(XrHandEXT Hand) const { return HandTrackers[Hand - 1]; }
+	bool ValidateXrHandTrackerEXT(XrHandTrackerEXT handTracker) { return (((FOXRVisionOSHandTracker*)handTracker) - ((FOXRVisionOSHandTracker*)&HandTrackers)) <= 1; }
+	void SyncHandTracking();
 };
diff --git a/Engine/Plugins/Runtime/OpenXRHandTracking/OpenXRHandTracking.uplugin b/Engine/Plugins/Runtime/OpenXRHandTracking/OpenXRHandTracking.uplugin
--- a/Engine/Plugins/Runtime/OpenXRHandTracking/OpenXRHandTracking.uplugin
+++ b/Engine/Plugins/Runtime/OpenXRHandTracking/OpenXRHandTracking.uplugin
@@ -27,7 +27,8 @@
 			"PlatformAllowList": [
 				"Win64",
 				"Linux",
-				"Android"
+				"Android",
+				"Mac"
 			]
 		},
 		{
@@ -36,7 +37,8 @@
 			"LoadingPhase": "Default",
 			"PlatformAllowList": [
 				"Win64",
-				"Linux"
+				"Linux",
+				"Mac"
 			]
 		}
 	],
@@ -47,7 +49,8 @@
 			"PlatformAllowList": [
 				"Win64",
 				"Linux",
-				"Android"
+				"Android",
+				"Mac"
 			]
 		},
 		{
@@ -55,4 +58,4 @@
 			"Enabled": true
 		}
 	]
-}
\ No newline at end of file
+}
