UE-208424 VisionOS: View lagging behind head motion
-Fixed problem where we were doing the copy of the swap chain render target into the compositor services drawable BEFORE rendering to the render target instead of after, resulted in displaying pixels from swapchainlength-1 in the past.  Now we are encoding the blits and the compositer services present into the main frame command buffer.
-Fixed problem where we were caching the device positions for late update before we update them in the render thread.
-Many verbose logging improvements.

diff --git a/Engine/Platforms/VisionOS/Plugins/Runtime/OpenXRVisionOS/Source/OXRVisionOS/Private/OXRVisionOS.h b/Engine/Platforms/VisionOS/Plugins/Runtime/OpenXRVisionOS/Source/OXRVisionOS/Private/OXRVisionOS.h
--- a/Engine/Platforms/VisionOS/Plugins/Runtime/OpenXRVisionOS/Source/OXRVisionOS/Private/OXRVisionOS.h
+++ b/Engine/Platforms/VisionOS/Plugins/Runtime/OpenXRVisionOS/Source/OXRVisionOS/Private/OXRVisionOS.h
@@ -22,19 +22,19 @@ static const TCHAR* VOSThreadString()
 {
 	if (IsInGameThread()) 
 	{
-		return TEXT("GT~");
+		return TEXT("T~G");
 	}
 	else if (IsInRenderingThread())
 	{
-		return TEXT("RT~");
+		return TEXT("T~R");
 	}
 	else if (IsInRHIThread())
 	{
-		return TEXT("IT~");
+		return TEXT("T~I");
 	}
 	else
 	{
-		return TEXT("?T~");
+		return TEXT("T~?");
 	}
 }
 
diff --git a/Engine/Platforms/VisionOS/Plugins/Runtime/OpenXRVisionOS/Source/OXRVisionOS/Private/OXRVisionOSSession.cpp b/Engine/Platforms/VisionOS/Plugins/Runtime/OpenXRVisionOS/Source/OXRVisionOS/Private/OXRVisionOSSession.cpp
--- a/Engine/Platforms/VisionOS/Plugins/Runtime/OpenXRVisionOS/Source/OXRVisionOS/Private/OXRVisionOSSession.cpp
+++ b/Engine/Platforms/VisionOS/Plugins/Runtime/OpenXRVisionOS/Source/OXRVisionOS/Private/OXRVisionOSSession.cpp
@@ -93,10 +93,6 @@ FOXRVisionOSSession::FOXRVisionOSSession(const XrSessionCreateInfo* createInfo,
 		bCreateFailed = true;
 		return;
 	}
-	
-	for (int32 i = 0; i < RenderToGameHeadTransformBufferLength; ++i) {
-		RenderToGameHeadTransform[i] = matrix_identity_float4x4;
-	}
 
 	//if (FOXRVisionOSController::Create(Controllers, this) == false)
 	//{
@@ -872,69 +868,62 @@ XrResult FOXRVisionOSSession::XrWaitFrame(
 	}
 
 	PipelinedFrameStateGame.FrameCounter = SessionFrameCounter++;
-	PipelinedFrameStateGame.RenderToGameHeadTransformIndexRead++;
-	UE_LOG(LogOXRVisionOS, Verbose, TEXT("%s FC: %i XrWaitFrameEvent->Wait() Finished"), VOSThreadString(), PipelinedFrameStateGame.FrameCounter);
-    
-    PipelinedFrameStateGame.SwiftFrame = cp_layer_renderer_query_next_frame(VOSLayerRenderer);
-    UE_LOG(LogOXRVisionOS, Verbose, TEXT("%s FC: %i SLF=0x%x cp_layer_renderer_query_next_frame in FOXRVisionOSSession::XrWaitFrame SwiftLayerFrame fetched"), VOSThreadString(), PipelinedFrameStateGame.FrameCounter, PipelinedFrameStateGame.SwiftFrame);
-
-    // Fetch the predicted timing information.
-    PipelinedFrameStateGame.SwiftFrameTiming = cp_frame_predict_timing(PipelinedFrameStateGame.SwiftFrame);
-    if ( PipelinedFrameStateGame.SwiftFrameTiming == nullptr)
-    {
-        //TODO ??? failed frame?
-        // This means the layer is not in the correct state  Perhaps we are either running too early or we need to end the session.
-        // we might want to return XR_ERROR_SESSION_NOT_RUNNING here, but only if we actually end the session first.
-        assert(false);
-        //return;
-    }
-
-    // Mark the beginnng of our game thread update work.  All the stuff before rendering.
-    UE_LOG(LogOXRVisionOS, Verbose, TEXT("%s FC: %i SLF: 0x%x cp_frame_start_update in FOXRVisionOSSession::XrWaitFrame"), VOSThreadString(), PipelinedFrameStateGame.FrameCounter, PipelinedFrameStateGame.SwiftFrame);
-    cp_frame_start_update(PipelinedFrameStateGame.SwiftFrame);
-
-    const cp_time_t PresentationTime = cp_frame_timing_get_presentation_time(PipelinedFrameStateGame.SwiftFrameTiming);
-    const CFTimeInterval TargetRenderTime = cp_time_to_cf_time_interval(PresentationTime); // CRTimeInterval is double time in seconds.
+	// Each frame we read this index on the game thread and then we write the index after this on the render thread.
+	PipelinedFrameStateGame.RenderToGameFrameStateIndex++;
+	UE_LOG(LogOXRVisionOS, VeryVerbose, TEXT("%s FC_%i XrWaitFrameEvent->Wait() Finished"), VOSThreadString(), PipelinedFrameStateGame.FrameCounter);
+    	
 
 	XrTime DisplayPeriod = 6 * 1000000; //TODO other frame sequences?  This is 90fps
 	frameState->predictedDisplayPeriod = DisplayPeriod;
-	frameState->predictedDisplayTime = OXRVisionOS::APITimeToXrTime(TargetRenderTime);
-	PipelinedFrameStateGame.PredictedDisplayTime = frameState->predictedDisplayTime;
 	
-	UE_LOG(LogOXRVisionOS, Verbose, TEXT("%s FC: %i SLF: 0x%x cp_frame_timing_get_presentation_time in FOXRVisionOSSession::XrWaitFrame TargetRenderTime =  %f xr predictedDisplayTime=%f"), VOSThreadString(), PipelinedFrameStateGame.FrameCounter, PipelinedFrameStateGame.SwiftFrame, (double)TargetRenderTime, (double)frameState->predictedDisplayTime);
-
-
-	PipelinedFrameStateGame.DeviceAnchor = ar_device_anchor_create();
-	check(ARKitWorldTrackingProvider)
-	{
-		//TODO: use a previous frame transform here, because we can't call this function from multiple threads
-		//TEMP use a lock.
-		FScopeLock ScopeLock(&CriticalSection_ar_world_tracking_provider_query_device_anchor_at_timestamp);
-		auto anchor_status = ar_world_tracking_provider_query_device_anchor_at_timestamp(ARKitWorldTrackingProvider, TargetRenderTime, PipelinedFrameStateGame.DeviceAnchor);
-		if (anchor_status == ar_device_anchor_query_status_success)
-		{
-			// Get the HMD transform and cache it
-			PipelinedFrameStateGame.HeadTransform = ar_anchor_get_origin_from_anchor_transform(PipelinedFrameStateGame.DeviceAnchor);
-		}
-		else
-		{
-			// Cache identity as the HMD transform as a fallback
-			PipelinedFrameStateGame.HeadTransform = matrix_identity_float4x4;
-		}
-	}
-
-//	PipelinedFrameStateGame.HeadTransform = RenderToGameHeadTransformRead(PipelinedFrameStateGame);
+	RenderToGameFrameStateRead(PipelinedFrameStateGame);
+	XrTime EarlyPredictionTime = OXRVisionOS::APITimeToXrTime(PipelinedFrameStateGame.PreviousPredictionTime);
+	EarlyPredictionTime += DisplayPeriod;
+	
+	frameState->predictedDisplayTime = EarlyPredictionTime;
 	
+	UE_LOG(LogOXRVisionOS, VeryVerbose, TEXT("%s FC_%i SLF_0x%x FOXRVisionOSSession::XrWaitFrame estimated XrTime predictedDisplayTime=%lld"), VOSThreadString(), PipelinedFrameStateGame.FrameCounter, PipelinedFrameStateGame.SwiftFrame, frameState->predictedDisplayTime);
+
     LocateViewInfoIndexAdvance();
 
 	frameState->shouldRender = SessionState == XR_SESSION_STATE_FOCUSED || SessionState == XR_SESSION_STATE_VISIBLE;
 
 	// We have to pipeline state updates in order to prevent task overlap (esp during map loading where we kick extra renders)
-	ENQUEUE_RENDER_COMMAND(UpdatePipelinedFrameState)([this, NewRenderingFrameState = PipelinedFrameStateGame](FRHICommandListImmediate& RHICmdList)
+	ENQUEUE_RENDER_COMMAND(UpdatePipelinedFrameState)([this, GameFrameState = PipelinedFrameStateGame](FRHICommandListImmediate& RHICmdList)
 	{
-		UE_LOG(LogOXRVisionOS, Verbose, TEXT("%s FC: %i SLF: 0x%x UpdatePipelinedFrameState G->R "), VOSThreadString(), NewRenderingFrameState.FrameCounter, NewRenderingFrameState.SwiftFrame);
+		UE_LOG(LogOXRVisionOS, VeryVerbose, TEXT("%s FC_%i  UpdatePipelinedFrameState G->R "), VOSThreadString(), GameFrameState.FrameCounter);
 		check(IsInRenderingThread());
-		PipelinedFrameStateRendering = NewRenderingFrameState;
+		
+		FPipelinedFrameState& FrameState = PipelinedFrameStateRendering;
+		FrameState = GameFrameState;
+		
+		FrameState.SwiftFrame = cp_layer_renderer_query_next_frame(VOSLayerRenderer);
+		UE_LOG(LogOXRVisionOS, VeryVerbose, TEXT("%s FC_%i SLF_0x%x  FOXRVisionOSSession::XrWaitFrame cp_layer_renderer_query_next_frame SwiftLayerFrame fetched"), VOSThreadString(), FrameState.FrameCounter, FrameState.SwiftFrame);
+
+		// Fetch the predicted timing information.
+		FrameState.SwiftFrameTiming = cp_frame_predict_timing(FrameState.SwiftFrame);
+		if ( FrameState.SwiftFrameTiming == nullptr)
+		{
+			//TODO ??? failed frame?
+			// This means the layer is not in the correct state  Perhaps we are either running too early or we need to end the session.
+			// we might want to return XR_ERROR_SESSION_NOT_RUNNING here, but only if we actually end the session first.
+			assert(false);
+			//return;
+		}
+		
+		CFTimeInterval TargetRenderTimeInterval = cp_time_to_cf_time_interval(cp_frame_timing_get_presentation_time(FrameState.SwiftFrameTiming));
+
+		UE_LOG(LogOXRVisionOS, VeryVerbose, TEXT("%s FC_%i SLF_0x%x FOXRVisionOSSession::XrWaitFrame cp_frame_predict_timing gave xr TargetRenderTimeInterval=%f"),
+			   VOSThreadString(), FrameState.FrameCounter, FrameState.SwiftFrame,
+			   (double)TargetRenderTimeInterval);
+
+		// Mark the beginnng of our game thread update work.  All the stuff before rendering.
+		UE_LOG(LogOXRVisionOS, VeryVerbose, TEXT("%s FC_%i SLF_0x%x FOXRVisionOSSession::XrWaitFrame cp_frame_start_update"), VOSThreadString(), FrameState.FrameCounter, FrameState.SwiftFrame);
+		cp_frame_start_update(FrameState.SwiftFrame);
+
+		FrameState.PredictedDisplayTime = OXRVisionOS::APITimeToXrTime(TargetRenderTimeInterval);
+		
+		cp_frame_end_update(FrameState.SwiftFrame);
 	});
 
 	// Update the STAGE space
@@ -946,18 +935,19 @@ XrResult FOXRVisionOSSession::XrWaitFrame(
 void FOXRVisionOSSession::OnBeginRendering_GameThread()
 {
     // Mark the end of our game frame
-    if (PipelinedFrameStateGame.SwiftFrame)
-    {
-		UE_LOG(LogOXRVisionOS, Verbose, TEXT("%s FC: %i SLF: 0x%x cp_frame_end_update in FOXRVisionOSSession::OnBeginRendering_GameThread"), VOSThreadString(), PipelinedFrameStateGame.FrameCounter, PipelinedFrameStateGame.SwiftFrame);
-        cp_frame_end_update(PipelinedFrameStateGame.SwiftFrame);
-        PipelinedFrameStateGame.SwiftFrame = nullptr;
-    }
+	// Currently we are just not tracking the update phase at all.
+//    if (PipelinedFrameStateGame.SwiftFrame)
+//    {
+//		UE_LOG(LogOXRVisionOS, VeryVerbose, TEXT("%s FC_%i SLF_0x%x cp_frame_end_update in FOXRVisionOSSession::OnBeginRendering_GameThread"), VOSThreadString(), PipelinedFrameStateGame.FrameCounter, PipelinedFrameStateGame.SwiftFrame);
+//        cp_frame_end_update(PipelinedFrameStateGame.SwiftFrame);
+//        PipelinedFrameStateGame.SwiftFrame = nullptr;
+//    }
 }
 
 // FOpenXRHMD::OnBeginRendering_RenderThread
 void FOXRVisionOSSession::OnBeginRendering_RenderThread()
 {
-	UE_LOG(LogOXRVisionOS, VeryVerbose, TEXT("%s FC: %i SLF: 0x%x FOXRVisionOSSession::OnBeginRendering_RenderThread"), VOSThreadString(), PipelinedFrameStateRendering.FrameCounter, PipelinedFrameStateRendering.SwiftFrame);
+	UE_LOG(LogOXRVisionOS, VeryVerbose, TEXT("%s FC_%i SLF_0x%x FOXRVisionOSSession::OnBeginRendering_RenderThread"), VOSThreadString(), PipelinedFrameStateRendering.FrameCounter, PipelinedFrameStateRendering.SwiftFrame);
 	
 	WaitUntil();
     StartSubmission();
@@ -966,7 +956,7 @@ void FOXRVisionOSSession::OnBeginRendering_RenderThread()
 	FRHICommandListImmediate& RHICmdList = FRHICommandListImmediate::Get();
 	RHICmdList.EnqueueLambda([this, NewRHIFrameState = PipelinedFrameStateRendering](FRHICommandListImmediate&)
 	{
-		UE_LOG(LogOXRVisionOS, VeryVerbose, TEXT("%s FC: %i SLF: 0x%x UpdatePipelinedFrameStat R->I, but HeadTransform Cleared"), VOSThreadString(), NewRHIFrameState.FrameCounter, NewRHIFrameState.SwiftFrame);
+		UE_LOG(LogOXRVisionOS, VeryVerbose, TEXT("%s FC_%i SLF_0x%x UpdatePipelinedFrameStat R->I, but HeadTransform Cleared"), VOSThreadString(), NewRHIFrameState.FrameCounter, NewRHIFrameState.SwiftFrame);
 		
 		PipelinedFrameStateRHI = NewRHIFrameState;
 		PipelinedFrameStateRHI.HeadTransform = matrix_identity_float4x4; // overwrite to ensure rhi gets a new one before it uses one.
@@ -981,7 +971,7 @@ XrResult FOXRVisionOSSession::XrBeginFrame(
 
 	SCOPED_NAMED_EVENT_TEXT("FOXRVisionOSSession::XrBeginFrame", FColor::Turquoise);
 	
-	UE_LOG(LogOXRVisionOS, VeryVerbose, TEXT("%s FC: %i SLF: 0x%x XrBeginFrame"), VOSThreadString(), PipelinedFrameStateRHI.FrameCounter, PipelinedFrameStateRHI.SwiftFrame);
+	UE_LOG(LogOXRVisionOS, VeryVerbose, TEXT("%s FC_%i SLF_0x%x SD_0x%8x XrBeginFrame"), VOSThreadString(), PipelinedFrameStateRHI.FrameCounter, PipelinedFrameStateRHI.SwiftFrame, PipelinedFrameStateRHI.SwiftDrawable);
 
 	// frameBeginInfo can be null
 	if (FrameBeginInfo != nullptr && FrameBeginInfo->type != XR_TYPE_FRAME_BEGIN_INFO)
@@ -1028,15 +1018,6 @@ XrResult FOXRVisionOSSession::XrBeginFrame(
         // But the docs infer that ReprojectionBeginFrame and waitUntilSafeForRendering are tied :/
     }
     
-    {
-		FPipelinedFrameState& FrameState = PipelinedFrameStateRHI;
-        
-        FrameState.SwiftDrawable = cp_frame_query_drawable(FrameState.SwiftFrame);
-        check(FrameState.SwiftDrawable);
-		UE_LOG(LogOXRVisionOS, VeryVerbose, TEXT("%s FC: %i SLF: 0x%x SD: 0x%x cp_frame_query_drawable in FOXRVisionOSSession::XrBeginFrame SwiftDrawable fetched"), VOSThreadString(), FrameState.FrameCounter, FrameState.SwiftFrame, FrameState.SwiftDrawable);
-
-    }
-    
     //MetalRHIOXRVisionOS::WaitUntilSafeForRendering(PipelinedFrameStateRHI.CommandListContext, NextBackBufferIndex);
     PreviousBackBufferIndexInBegin = NextBackBufferIndex;
 
@@ -1047,73 +1028,67 @@ void FOXRVisionOSSession::WaitUntil()
 {
     FPipelinedFrameState& FrameState = GetPipelinedFrameStateForThread();
     
-	UE_LOG(LogOXRVisionOS, VeryVerbose, TEXT("%s FC: %i SLF: 0x%x SD: 0x%x cp_time_wait_until in FOXRVisionOSSession::WaitUntil started"), VOSThreadString(), FrameState.FrameCounter, FrameState.SwiftFrame, FrameState.SwiftDrawable);
-	
-    cp_time_wait_until(cp_frame_timing_get_optimal_input_time(FrameState.SwiftFrameTiming));
+	cp_time_t OptimalInputTime = cp_frame_timing_get_optimal_input_time(FrameState.SwiftFrameTiming);
+	CFTimeInterval TimeInterval = cp_time_to_cf_time_interval(OptimalInputTime);
 
-	UE_LOG(LogOXRVisionOS, VeryVerbose, TEXT("%s FC: %i SLF: 0x%x SD: 0x%x cp_time_wait_until in FOXRVisionOSSession::WaitUntil completed"), VOSThreadString(), FrameState.FrameCounter, FrameState.SwiftFrame, FrameState.SwiftDrawable);
+	UE_LOG(LogOXRVisionOS, VeryVerbose, TEXT("%s FC_%i SLF_0x%x SD_0x%8x FOXRVisionOSSession::WaitUntil cp_time_wait_until started waiting until OptimalInputTime %f"), VOSThreadString(), FrameState.FrameCounter, FrameState.SwiftFrame, FrameState.SwiftDrawable, TimeInterval);
+	
+    cp_time_wait_until(OptimalInputTime);
+	
+	UE_LOG(LogOXRVisionOS, VeryVerbose, TEXT("%s FC_%i SLF_0x%x SD_0x%8x FOXRVisionOSSession::WaitUntil cp_time_wait_until completed"), VOSThreadString(), FrameState.FrameCounter, FrameState.SwiftFrame, FrameState.SwiftDrawable);
 }
 
 void FOXRVisionOSSession::StartSubmission()
 {
     FPipelinedFrameState& FrameState = GetPipelinedFrameStateForThread();
 	
-	UE_LOG(LogOXRVisionOS, VeryVerbose, TEXT("%s FC: %i SLF: 0x%x SD: 0x%x cp_frame_start_submission in FOXRVisionOSSession::StartSubmission"), VOSThreadString(), FrameState.FrameCounter, FrameState.SwiftFrame, FrameState.SwiftDrawable);
+	UE_LOG(LogOXRVisionOS, VeryVerbose, TEXT("%s FC_%i SLF_0x%x SD_0x%8x FOXRVisionOSSession::StartSubmission cp_frame_start_submission"), VOSThreadString(), FrameState.FrameCounter, FrameState.SwiftFrame, FrameState.SwiftDrawable);
     
     cp_frame_start_submission(FrameState.SwiftFrame);
 
     FrameState.SwiftDrawable = cp_frame_query_drawable(FrameState.SwiftFrame);
     check(FrameState.SwiftDrawable);
 	
-	UE_LOG(LogOXRVisionOS, VeryVerbose, TEXT("%s FC: %i SLF: 0x%x SD: 0x%x cp_frame_query_drawable in FOXRVisionOSSession::StartSubmission got the SD again."), VOSThreadString(), FrameState.FrameCounter, FrameState.SwiftFrame, FrameState.SwiftDrawable);
+	UE_LOG(LogOXRVisionOS, VeryVerbose, TEXT("%s FC_%i SLF_0x%x SD_0x%8x FOXRVisionOSSession::StartSubmission cp_frame_query_drawable in got the SD again."), VOSThreadString(), FrameState.FrameCounter, FrameState.SwiftFrame, FrameState.SwiftDrawable);
 
     FrameState.SwiftFinalFrameTiming = cp_drawable_get_frame_timing(FrameState.SwiftDrawable);
+	FrameState.SwiftFinalFrameTimeInterval = cp_time_to_cf_time_interval(cp_frame_timing_get_presentation_time(FrameState.SwiftFinalFrameTiming));
+	FrameState.DeviceAnchor = ar_device_anchor_create();
 	
-	//FrameState.DeviceAnchor = ar_device_anchor_create();
-	
-    CFTimeInterval TimeInterval = cp_time_to_cf_time_interval(cp_frame_timing_get_presentation_time(FrameState.SwiftFinalFrameTiming));
-	
-	UE_LOG(LogOXRVisionOS, VeryVerbose, TEXT("%s FC: %i SLF: 0x%x SD: 0x%x cp_drawable_get_frame_timing in FOXRVisionOSSession::StartSubmission got SwiftFinalFrameTiming %f"), VOSThreadString(), FrameState.FrameCounter, FrameState.SwiftFrame, FrameState.SwiftDrawable, (double)TimeInterval);
+	UE_LOG(LogOXRVisionOS, VeryVerbose, TEXT("%s FC_%i SLF_0x%x SD_0x%8x FOXRVisionOSSession::StartSubmission cp_drawable_get_frame_timing got SwiftFinalFrameTimeInterval %f"), VOSThreadString(), FrameState.FrameCounter, FrameState.SwiftFrame, FrameState.SwiftDrawable, (double)FrameState.SwiftFinalFrameTimeInterval);
+
+	uint64_t CurrentMachTime = mach_absolute_time();
+	uint64_t PredictedPresentationMachTime = cp_frame_timing_get_presentation_time(FrameState.SwiftFinalFrameTiming).cp_mach_abs_time;
+	UE_LOG(LogOXRVisionOS, VeryVerbose, TEXT("   FC_%i                   FOXRVisionOSSession::StartSubmission CurrentMachTime %lld PredictedPresentationMachTime %lld"), FrameState.FrameCounter,
+		   CurrentMachTime, PredictedPresentationMachTime);
+
+	check(ARKitWorldTrackingProvider);
 	
-    check(ARKitWorldTrackingProvider)
 	{
-		//TEMP use a lock.
-		FScopeLock ScopeLock(&CriticalSection_ar_world_tracking_provider_query_device_anchor_at_timestamp);
-		auto anchor_status = ar_world_tracking_provider_query_device_anchor_at_timestamp(ARKitWorldTrackingProvider, TimeInterval, FrameState.DeviceAnchor);
+		auto anchor_status = ar_world_tracking_provider_query_device_anchor_at_timestamp(ARKitWorldTrackingProvider, FrameState.SwiftFinalFrameTimeInterval, FrameState.DeviceAnchor);
 		if (anchor_status == ar_device_anchor_query_status_success)
 		{
-//			if (OXRVisionOSCVars::CVarDisableReprojection.GetValueOnAnyThread())
-//			{
-//				// Note: Documentation says this will disable reprojection, but it does not apper to actually work:
-				// Log says: "Presenting a drawable without a device anchor. This drawable won't be presented."
-//				cp_drawable_set_device_anchor(FrameState.SwiftDrawable, nullptr);
-//			}
-//			else
-			{
-				cp_drawable_set_device_anchor(FrameState.SwiftDrawable, FrameState.DeviceAnchor);
-			}
-
+			cp_drawable_set_device_anchor(FrameState.SwiftDrawable, FrameState.DeviceAnchor);
+		
 			// Get the HMD transform and cache it
 			FrameState.HeadTransform = ar_anchor_get_origin_from_anchor_transform(FrameState.DeviceAnchor);
-			UE_LOG(LogOXRVisionOS, VeryVerbose, TEXT("%s FC: %i SLF: 0x%x SD: 0x%x cp_drawable_set_device_anchor in FOXRVisionOSSession::StartSubmission with DeviceAnchor 0x%x  also getting head transform and buffering it for GT HeadTransform01=%f"), VOSThreadString(), FrameState.FrameCounter, FrameState.SwiftFrame, FrameState.SwiftDrawable, FrameState.DeviceAnchor, FrameState.HeadTransform.columns[0][1]);
-			
 		}
 		else
 		{
-			UE_LOG(LogOXRVisionOS, VeryVerbose, TEXT("%s FC: %i SLF: 0x%x SD: 0x%x cp_drawable_set_device_anchor in FOXRVisionOSSession::StartSubmission DeviceAnchor query failed, setting nullptr, no reprojection"), VOSThreadString(), FrameState.FrameCounter, FrameState.SwiftFrame, FrameState.SwiftDrawable);
+			UE_LOG(LogOXRVisionOS, VeryVerbose, TEXT("%s FC_%i SLF_0x%x SD_0x%8x  FOXRVisionOSSession::StartSubmission cp_drawable_set_device_anchor DeviceAnchor query failed, setting nullptr, no reprojection"), VOSThreadString(), FrameState.FrameCounter, FrameState.SwiftFrame, FrameState.SwiftDrawable);
 			
 			cp_drawable_set_device_anchor(FrameState.SwiftDrawable, nullptr);
 			// Leave FrameState.HeadTransform untouched, same as game thread. It may simply be identity.  But whatever it is we will go ahead and render with it.
 		}
-		RenderToGameHeadTransformWrite(FrameState);
 	}
+	RenderToGameFrameStateWrite(FrameState);
 }
 
 // FOpenXRHMD::OnFinishRendering_RHIThread()
 XrResult FOXRVisionOSSession::XrEndFrame(
 	const XrFrameEndInfo* InFrameEndInfo)
 {
-	UE_LOG(LogOXRVisionOS, Verbose, TEXT("%s FC: %i SLF: 0x%x SD: 0x%x FOXRVisionOSSession::XrEndFrame"), VOSThreadString(), PipelinedFrameStateRHI.FrameCounter, PipelinedFrameStateRHI.SwiftFrame, PipelinedFrameStateRHI.SwiftDrawable);
+	UE_LOG(LogOXRVisionOS, VeryVerbose, TEXT("%s FC_%i SLF_0x%x SD_0x%8x FOXRVisionOSSession::XrEndFrame"), VOSThreadString(), PipelinedFrameStateRHI.FrameCounter, PipelinedFrameStateRHI.SwiftFrame, PipelinedFrameStateRHI.SwiftDrawable);
 
 	SCOPED_NAMED_EVENT_TEXT("FOXRVisionOSSession::XrEndFrame", FColor::Turquoise);
 
@@ -1155,9 +1130,9 @@ XrResult FOXRVisionOSSession::XrEndFrame(
         check(View0.subImage.swapchain == View1.subImage.swapchain); //TODO: should we do it this way or should we have two separate swapchains?  If we did we would need to handle both below.
         
         const FOXRVisionOSSwapchain& SwapchainData0 = *(reinterpret_cast<FOXRVisionOSSwapchain*>(View0.subImage.swapchain));
-        const FOXRVisionOSSwapchain::FSwapchainImage& ReleasedImage = SwapchainData0.GetLastReleasedImage();
+        const FOXRVisionOSSwapchain::FSwapchainImage& WaitedImage = SwapchainData0.GetLastWaitedImage();
 		
-		UE_LOG(LogOXRVisionOS, VeryVerbose, TEXT("%s FC: %i SLF: 0x%x SD: 0x%x FOXRVisionOSSession::XrEndFrame SwapchainData0.GetLastReleasedImage() = 0x%x"), VOSThreadString(), PipelinedFrameStateRHI.FrameCounter, PipelinedFrameStateRHI.SwiftFrame, PipelinedFrameStateRHI.SwiftDrawable, ReleasedImage.Image.GetReference()->GetNativeResource());
+		UE_LOG(LogOXRVisionOS, VeryVerbose, TEXT("%s FC_%i SLF_0x%x SD_0x%8x FOXRVisionOSSession::XrEndFrame SwapchainData0.GetLastWaitedImage() metal textgure = 0x%x"), VOSThreadString(), PipelinedFrameStateRHI.FrameCounter, PipelinedFrameStateRHI.SwiftFrame, PipelinedFrameStateRHI.SwiftDrawable, WaitedImage.Image.GetReference()->GetNativeResource());
 		
 		const XrCompositionLayerDepthInfoKHR* Depth0 = OpenXR::FindChainedStructByType<XrCompositionLayerDepthInfoKHR>(View0.next, (XrStructureType)XR_TYPE_COMPOSITION_LAYER_DEPTH_INFO_KHR);
 		const XrCompositionLayerDepthInfoKHR* Depth1 = OpenXR::FindChainedStructByType<XrCompositionLayerDepthInfoKHR>(View1.next, (XrStructureType)XR_TYPE_COMPOSITION_LAYER_DEPTH_INFO_KHR);
@@ -1166,9 +1141,9 @@ XrResult FOXRVisionOSSession::XrEndFrame(
 		check(Depth0->subImage.swapchain == Depth1->subImage.swapchain); //TODO: should we do it this way or should we have two separate swapchains?  If we did we would need to handle both below.
 
 		const FOXRVisionOSSwapchain& SwapchainDepthData0 = *(reinterpret_cast<FOXRVisionOSSwapchain*>(Depth0->subImage.swapchain));
-		const FOXRVisionOSSwapchain::FSwapchainImage& ReleasedDepth = SwapchainDepthData0.GetLastReleasedImage();
+		const FOXRVisionOSSwapchain::FSwapchainImage& WaitedDepth = SwapchainDepthData0.GetLastWaitedImage();
 
-        const MetalRHIVisionOS::PresentImmersiveParams Params{ReleasedImage.Image, ReleasedDepth.Image, PipelinedFrameStateRHI.SwiftFrame, PipelinedFrameStateRHI.SwiftDrawable};
+        const MetalRHIVisionOS::PresentImmersiveParams Params{WaitedImage.Image, WaitedDepth.Image, PipelinedFrameStateRHI.SwiftFrame, PipelinedFrameStateRHI.SwiftDrawable, PipelinedFrameStateRHI.FrameCounter};
         MetalRHIVisionOS::PresentImmersive(Params);
     }
 
@@ -1211,7 +1186,7 @@ XrResult FOXRVisionOSSession::XrLocateViews(
     
 	FPipelinedFrameState& FrameState = GetPipelinedFrameStateForThread();
 	
-	UE_LOG(LogOXRVisionOS, Verbose, TEXT("%s FC: %i SLF: 0x%x SD: 0x%x FOXRVisionOSSession::XrLocateViews"), VOSThreadString(), FrameState.FrameCounter, FrameState.SwiftFrame, FrameState.SwiftDrawable);
+	//UE_LOG(LogOXRVisionOS, VeryVerbose, TEXT("%s FC_%i SLF_0x%x SD_0x%8x FOXRVisionOSSession::XrLocateViews"), VOSThreadString(), FrameState.FrameCounter, FrameState.SwiftFrame, FrameState.SwiftDrawable);
 
 	uint32_t NumViews = 0;
     if (IsInGameThread())
@@ -1252,14 +1227,11 @@ XrResult FOXRVisionOSSession::XrLocateViews(
                 const FLocateViewInfo& ViewInfoCache = GetLocateViewInfo_GameThread();
                 Views[Index].pose = OXRVisionOS::ToXrPose(ViewInfoCache.ViewTransforms[Index]);
                 Views[Index].fov = ViewInfoCache.HmdFovs[Index];
-				UE_LOG(LogOXRVisionOS, Verbose, TEXT("%s FC: %i SLF: 0x%x SD: 0x%x FOXRVisionOSSession::XrLocateViews cached HeadTransform01=%f xrposexy=%f,%f"), VOSThreadString(), FrameState.FrameCounter, FrameState.SwiftFrame, FrameState.SwiftDrawable, FrameState.HeadTransform.columns[0][1], Views[Index].pose.position.x, Views[Index].pose.position.y);
             }
             else
             {
                 check(IsInRenderingThread());
                 check(FrameState.SwiftDrawable);
-                
-				UE_LOG(LogOXRVisionOS, VeryVerbose, TEXT("%s FC: %i SLF: 0x%x SD: 0x%x cp_drawable_get_view in FOXRVisionOSSession::XrLocateViews view %i"), VOSThreadString(), FrameState.FrameCounter, FrameState.SwiftFrame, FrameState.SwiftDrawable, Index);
 
                 cp_view_t View = cp_drawable_get_view(FrameState.SwiftDrawable, Index);
                 simd_float4 Tangents = cp_view_get_tangents(View);
@@ -1269,7 +1241,7 @@ XrResult FOXRVisionOSSession::XrLocateViews(
     //                                                                          true); /* reverseZ */
                 
                 FrameState.HmdFovs[Index].angleLeft     = -FMath::Atan(Tangents[0]);
-                FrameState.HmdFovs[Index].angleRight    =  FMath::Atan(Tangents[1]);
+				FrameState.HmdFovs[Index].angleRight    =  FMath::Atan(Tangents[1]);
                 FrameState.HmdFovs[Index].angleUp       =  FMath::Atan(Tangents[2]);
                 FrameState.HmdFovs[Index].angleDown     = -FMath::Atan(Tangents[3]);
 
@@ -1279,8 +1251,6 @@ XrResult FOXRVisionOSSession::XrLocateViews(
                 simd_float4x4 ViewTransform = cp_view_get_transform(View);
                 Views[Index].pose = OXRVisionOS::ToXrPose(ViewTransform);
                 Views[Index].fov = FrameState.HmdFovs[Index];
-				
-				UE_LOG(LogOXRVisionOS, Verbose, TEXT("%s FC: %i SLF: 0x%x SD: 0x%x FOXRVisionOSSession::XrLocateViews cp_view_get_transformHeadTransform01=%f xrposexy=%f,%f"), VOSThreadString(), FrameState.FrameCounter, FrameState.SwiftFrame, FrameState.SwiftDrawable, FrameState.HeadTransform.columns[0][1], Views[Index].pose.position.x, Views[Index].pose.position.y);
 
                 // Cache the render thread data for the game thread
                 FLocateViewInfo& ViewInfoCache = GetLocateViewInfo_RenderThread();
@@ -1306,7 +1276,10 @@ void FOXRVisionOSSession::GetHMDTransform(XrTime DisplayTime, FTransform& OutTra
     
     OutTransform = OXRVisionOS::ToFTransform(FrameState.HeadTransform);
 
-	UE_LOG(LogOXRVisionOS, VeryVerbose, TEXT("%s FC: %i SLF: 0x%x SD: 0x%x FOXRVisionOSSession::GetHMDTransform HeadTransform01=%f"), VOSThreadString(), FrameState.FrameCounter, FrameState.SwiftFrame, FrameState.SwiftDrawable, FrameState.HeadTransform.columns[0][1]);
+//	UE_LOG(LogOXRVisionOS, VeryVerbose, TEXT("%s FC_%i SLF_0x%x SD_0x%8x FOXRVisionOSSession::GetHMDTransform XrHeadTransform x,y=%f,%f r=%f OutTransform x,y=%f,%f"),
+//		   VOSThreadString(), FrameState.FrameCounter, FrameState.SwiftFrame, FrameState.SwiftDrawable,
+//		   FrameState.HeadTransform.columns[3][0], FrameState.HeadTransform.columns[3][1], FrameState.HeadTransform.columns[0][0],
+//		   OutTransform.GetLocation().X, OutTransform.GetLocation().Y);
 }
 
 void FOXRVisionOSSession::GetTransformForButton(EOXRVisionOSControllerButton Button, XrTime DisplayTime, FTransform& OutTransform, XrSpaceLocationFlags& OutLocationFlags, FVector& OutLinearVelocity, FVector& OutAngularVelocity, XrSpaceVelocityFlags& OutVelocityFlags)
@@ -1453,15 +1426,25 @@ FOXRVisionOSSession::FPipelinedFrameState& FOXRVisionOSSession::GetPipelinedFram
 }
 
 // Read from the current index.
-simd_float4x4 FOXRVisionOSSession::RenderToGameHeadTransformRead(const FPipelinedFrameState& FrameState)
+void FOXRVisionOSSession::RenderToGameFrameStateRead(FPipelinedFrameState& InOutFrameState)
 {
-	return RenderToGameHeadTransform[FrameState.RenderToGameHeadTransformIndexRead % RenderToGameHeadTransformBufferLength];
+	FOXRVisionOSSession::FRenderToGameFrameState& RTGFS = RenderToGameFrameState[InOutFrameState.RenderToGameFrameStateIndex % RenderToGameFrameStateBufferLength];
+	InOutFrameState.HeadTransform = RTGFS.HeadTransform;
+	InOutFrameState.PreviousPredictionTime = RTGFS.PredictionTime;
+	
+//	UE_LOG(LogOXRVisionOS, VeryVerbose, TEXT("%s FC_%i SLF_0x%x SD_0x%8x  FOXRVisionOSSession::RenderToGameFrameStateRead HeadTransform01 x,y=%f,%f r=%f  PredictionTime=%f"),
+//		   VOSThreadString(), InOutFrameState.FrameCounter, InOutFrameState.SwiftFrame, InOutFrameState.SwiftDrawable,
+//		   RTGFS.HeadTransform.columns[3][0], RTGFS.HeadTransform.columns[3][1], RTGFS.HeadTransform.columns[0][0], (double)RTGFS.PredictionTime);
 }
 
-// Write to the next index that the game thread will read from.
-void FOXRVisionOSSession::RenderToGameHeadTransformWrite(FPipelinedFrameState& FrameState)
+// Write into the next index that the game thread will read.
+void FOXRVisionOSSession::RenderToGameFrameStateWrite(const FPipelinedFrameState& FrameState)
 {
-	RenderToGameHeadTransform[(FrameState.RenderToGameHeadTransformIndexRead + 1) % RenderToGameHeadTransformBufferLength] = FrameState.HeadTransform;
+	FOXRVisionOSSession::FRenderToGameFrameState& RTGFS = RenderToGameFrameState[(FrameState.RenderToGameFrameStateIndex + 1) % RenderToGameFrameStateBufferLength];
+	RTGFS.HeadTransform = FrameState.HeadTransform;
+	RTGFS.PredictionTime = FrameState.SwiftFinalFrameTimeInterval;
+	
+	//UE_LOG(LogOXRVisionOS, VeryVerbose, TEXT("%s FC_%i SLF_0x%x SD_0x%8x  FOXRVisionOSSession::RenderToGameFrameStateWrite HeadTransform01 x,y=%f,%f  PredictionTime=%f"), VOSThreadString(), FrameState.FrameCounter, FrameState.SwiftFrame, FrameState.SwiftDrawable, RTGFS.HeadTransform.columns[3][0], RTGFS.HeadTransform.columns[3][1], (double)RTGFS.PredictionTime);
 }
 
 
@@ -1497,7 +1480,6 @@ XrResult FOXRVisionOSSession::FOXRVisionOSHandTracker::XrLocateHandJointsEXT(
 	}
 	
 	XrPosef pose = locations->jointLocations[XR_HAND_JOINT_THUMB_TIP_EXT].pose;
-	UE_LOG(LogTemp, Log, TEXT("XrLocateHandJointsEXT() locating thumb tip flags: %i  xyz %0.2f, %0.2f, %0.2f"), locations->jointLocations[XR_HAND_JOINT_THUMB_TIP_EXT].locationFlags, pose.position.x, pose.position.y, pose.position.z);
 
 	return XR_SUCCESS;
 }
diff --git a/Engine/Platforms/VisionOS/Plugins/Runtime/OpenXRVisionOS/Source/OXRVisionOS/Private/OXRVisionOSSession.h b/Engine/Platforms/VisionOS/Plugins/Runtime/OpenXRVisionOS/Source/OXRVisionOS/Private/OXRVisionOSSession.h
--- a/Engine/Platforms/VisionOS/Plugins/Runtime/OpenXRVisionOS/Source/OXRVisionOS/Private/OXRVisionOSSession.h
+++ b/Engine/Platforms/VisionOS/Plugins/Runtime/OpenXRVisionOS/Source/OXRVisionOS/Private/OXRVisionOSSession.h
@@ -158,19 +158,26 @@ public:
 		//FTransform HMDPoseInTrackerSpace = FTransform::Identity;
 		//PoseData HmdPose = {};
 		cp_frame_t SwiftFrame = nullptr;
-        cp_frame_timing_t SwiftFrameTiming = nullptr;
-        simd_float4x4 HeadTransform = matrix_identity_float4x4;
-        cp_drawable_t SwiftDrawable = nullptr;
+		cp_frame_timing_t SwiftFrameTiming = nullptr;  // Used for the cp_time_wait_until
+		CFTimeInterval PreviousPredictionTime = 0;
+		cp_drawable_t SwiftDrawable = nullptr;
         cp_frame_timing_t SwiftFinalFrameTiming = nullptr; // TODO should we do this or should we overwrite SwifFrameTiming???
+		simd_float4x4 HeadTransform = matrix_identity_float4x4;
+		CFTimeInterval SwiftFinalFrameTimeInterval = 0;
         ar_device_anchor_t DeviceAnchor = nullptr;
 		bool bSynchronizing = true; // when true we are in the process of starting up a new session and the new frame state has not yet propagated.
 		XrFovf HmdFovs[2];
         int32 LocateViewInfoBufferIndex = 0;
-		int32 RenderToGameHeadTransformIndexRead = -1;
+		int32 RenderToGameFrameStateIndex = -1;
 	};
 	
-	simd_float4x4 RenderToGameHeadTransformRead(const FPipelinedFrameState& FrameState);
-	void RenderToGameHeadTransformWrite(FPipelinedFrameState& FrameState);
+	struct FRenderToGameFrameState
+	{
+		simd_float4x4 HeadTransform = matrix_identity_float4x4;
+		CFTimeInterval PredictionTime = 0;
+	};
+	void RenderToGameFrameStateRead(FPipelinedFrameState& InOutFrameState);
+	void RenderToGameFrameStateWrite(const FPipelinedFrameState& FrameState);
 
 	const FPipelinedFrameState& GetPipelinedFrameStateForThread() const;
 	FPipelinedFrameState& GetPipelinedFrameStateForThread();
@@ -235,10 +242,8 @@ private:
 	FPipelinedFrameState	PipelinedFrameStateRHI;
 	int32 CachedBeginFlipFrameCounter = INT32_MAX;
 	
-	static const int32 RenderToGameHeadTransformBufferLength = 3;
-	simd_float4x4 RenderToGameHeadTransform[RenderToGameHeadTransformBufferLength];
-	//TEMP see usages
-	FCriticalSection CriticalSection_ar_world_tracking_provider_query_device_anchor_at_timestamp;
+	static const int32 RenderToGameFrameStateBufferLength = 3;
+	FRenderToGameFrameState RenderToGameFrameState[RenderToGameFrameStateBufferLength];
 
 	FOpenXRRenderBridge* RenderBridge = nullptr;
 
diff --git a/Engine/Platforms/VisionOS/Source/Runtime/Launch/Source/UESwift.swift b/Engine/Platforms/VisionOS/Source/Runtime/Launch/Source/UESwift.swift
--- a/Engine/Platforms/VisionOS/Source/Runtime/Launch/Source/UESwift.swift
+++ b/Engine/Platforms/VisionOS/Source/Runtime/Launch/Source/UESwift.swift
@@ -115,6 +115,8 @@ struct UESwiftApp: App {
 					FSwiftAppBootstrap.KickoffWithCompositingLayer(layerRenderer)
 			}
 		}
+		//.upperLimbVisibility(.hidden)
+		.upperLimbVisibility(.visible)
 	}
 }
 
diff --git a/Engine/Plugins/Runtime/OpenXR/Source/OpenXRHMD/Private/OpenXRHMD.cpp b/Engine/Plugins/Runtime/OpenXR/Source/OpenXRHMD/Private/OpenXRHMD.cpp
--- a/Engine/Plugins/Runtime/OpenXR/Source/OpenXRHMD/Private/OpenXRHMD.cpp
+++ b/Engine/Plugins/Runtime/OpenXR/Source/OpenXRHMD/Private/OpenXRHMD.cpp
@@ -50,6 +50,25 @@
 #include "Misc/MessageDialog.h"
 #include "UnrealEdMisc.h"
 #endif
+static const TCHAR* HMDThreadString()
+{
+	if (IsInGameThread())
+	{
+		return TEXT("T~G");
+	}
+	else if (IsInRenderingThread())
+	{
+		return TEXT("T~R");
+	}
+	else if (IsInRHIThread())
+	{
+		return TEXT("T~I");
+	}
+	else
+	{
+		return TEXT("T~?");
+	}
+}
 
 #define LOCTEXT_NAMESPACE "OpenXR"
 
@@ -1965,7 +1984,7 @@ bool FOpenXRHMD::OnStereoStartup()
 	}
 
 	bUseCustomReferenceSpace = false;
-	XrReferenceSpaceType CustomReferenceSpaceType;
+	XrReferenceSpaceType CustomReferenceSpaceType = XR_REFERENCE_SPACE_TYPE_LOCAL;
 	for (IOpenXRExtensionPlugin* Module : ExtensionPlugins)
 	{
 		if (Module->UseCustomReferenceSpaceType(CustomReferenceSpaceType))
@@ -2342,6 +2361,7 @@ IStereoRenderTargetManager* FOpenXRHMD::GetRenderTargetManager()
 
 int32 FOpenXRHMD::AcquireColorTexture()
 {
+	check(IsInGameThread());
 	if (Session)
 	{
 		const FXRSwapChainPtr& ColorSwapchain = PipelinedLayerStateRendering.ColorSwapchain;
@@ -2916,7 +2936,10 @@ void FOpenXRHMD::OnBeginRendering_RenderThread(FRHICommandListImmediate& RHICmdL
 	{
 		Module->OnBeginRendering_RenderThread(Session);
 	}
-
+	
+	// Snapshot new poses for late update.
+	UpdateDeviceLocations(false);
+	
 	SetupFrameLayers_RenderThread(RHICmdList);
 
 	const float WorldToMeters = GetWorldToMetersScale();
@@ -2988,7 +3011,7 @@ void FOpenXRHMD::OnBeginRendering_RenderThread(FRHICommandListImmediate& RHICmdL
 			FBFoveationImageGenerator->SetCurrentFrameSwapchainIndex(ColorSwapchain->GetSwapChainIndex_RHIThread());
 		}
 
-		UE_LOG(LogHMD, VeryVerbose, TEXT("EnqueueLambda OnBeginRendering_RHIThread WaitCount: %i"), PipelinedFrameStateRendering.WaitCount);
+		UE_LOG(LogHMD, VeryVerbose, TEXT("%s WF_%i EnqueueLambda OnBeginRendering_RHIThread"), HMDThreadString(), PipelinedFrameStateRendering.WaitCount);
 		RHICmdList.EnqueueLambda([this, FrameState = PipelinedFrameStateRendering, ColorSwapchain, DepthSwapchain, EmulationSwapchain](FRHICommandListImmediate& InRHICmdList)
 		{
 			OnBeginRendering_RHIThread(FrameState, ColorSwapchain, DepthSwapchain, EmulationSwapchain);
@@ -3087,14 +3110,9 @@ void FOpenXRHMD::OnBeginRendering_GameThread()
 			UE_CLOG(PipelinedFrameStateRendering.FrameState.predictedDisplayTime >= GameFrameState.FrameState.predictedDisplayTime,
 				LogHMD, VeryVerbose, TEXT("Predicted display time went backwards from %lld to %lld"), PipelinedFrameStateRendering.FrameState.predictedDisplayTime, GameFrameState.FrameState.predictedDisplayTime);
 
-			UE_LOG(LogHMD, VeryVerbose, TEXT("FOpenXRHMD TransferFrameStateToRenderingThread %i"), GameFrameState.WaitCount);
+			UE_LOG(LogHMD, VeryVerbose, TEXT("%s WF_%i FOpenXRHMD TransferFrameStateToRenderingThread"), HMDThreadString(), GameFrameState.WaitCount);
 			PipelinedFrameStateRendering = GameFrameState;
-
-			// Snapshot new poses for late update.
-			// We do this here instead of in OnBeginRendering_RenderThread in order to have this ready before Scene Captures and Reflection
-			// Captures render, as these may need to use the HMDs location to calculate accurate view matrices.
-			UpdateDeviceLocations(false);
-
+			
 			PipelinedLayerStateRendering.LayerStateFlags = EOpenXRLayerStateFlags::None;
 
 			// If we are emulating layers, we still need to submit background layer since we composite into it
@@ -3137,9 +3155,9 @@ void FOpenXRHMD::OnBeginSimulation_GameThread()
 	}
 	static int WaitCount = 0;
 	++WaitCount;
-	UE_LOG(LogHMD, VeryVerbose, TEXT("xrWaitFrame %i"), WaitCount);
+	UE_LOG(LogHMD, VeryVerbose, TEXT("%s WF_%i xrWaitFrame Calling..."), HMDThreadString(), WaitCount);
 	XR_ENSURE(xrWaitFrame(Session, &WaitInfo, &FrameState));
-	UE_LOG(LogHMD, VeryVerbose, TEXT("xrWaitFrame %i Complete"), WaitCount);
+	UE_LOG(LogHMD, VeryVerbose, TEXT("%s WF_%i xrWaitFrame Complete"), HMDThreadString(), WaitCount);
 
 	// The pipeline state on the game thread can only be safely modified after xrWaitFrame which will be unblocked by
 	// the runtime when xrBeginFrame is called. The rendering thread will clone the game pipeline state before calling
@@ -3425,7 +3443,7 @@ void FOpenXRHMD::OnBeginRendering_RHIThread(const FPipelinedFrameState& InFrameS
 	}
 	static int BeginCount = 0;
 	PipelinedFrameStateRHI.BeginCount = ++BeginCount;
-	UE_LOG(LogHMD, VeryVerbose, TEXT("xrBeginFrame WaitCount: %i BeginCount: %i"), PipelinedFrameStateRHI.WaitCount, PipelinedFrameStateRHI.BeginCount);
+	UE_LOG(LogHMD, VeryVerbose, TEXT("%s WF_%i xrBeginFrame BeginCount: %i"), HMDThreadString(), PipelinedFrameStateRHI.WaitCount, PipelinedFrameStateRHI.BeginCount);
 	XrResult Result = xrBeginFrame(Session, &BeginInfo);
 	if (XR_SUCCEEDED(Result))
 	{
@@ -3465,7 +3483,9 @@ void FOpenXRHMD::OnBeginRendering_RHIThread(const FPipelinedFrameState& InFrameS
 
 		bIsRendering = true;
 
-		UE_LOG(LogHMD, VeryVerbose, TEXT("Rendering frame predicted to be displayed at %lld"), InFrameState.FrameState.predictedDisplayTime);
+		UE_LOG(LogHMD, VeryVerbose, TEXT("%s WF_%i Rendering frame predicted to be displayed at %lld"), 
+			   HMDThreadString(), PipelinedFrameStateRHI.WaitCount,
+			   PipelinedFrameStateRHI.FrameState.predictedDisplayTime);
 	}
 	else
 	{
@@ -3488,6 +3508,8 @@ void FOpenXRHMD::OnFinishRendering_RHIThread()
 	{
 		return;
 	}
+	
+	UE_LOG(LogHMD, VeryVerbose, TEXT("%s WF_%i FOpenXRHMD::OnFinishRendering_RHIThread releasing swapchain images now."), HMDThreadString(), PipelinedFrameStateRHI.WaitCount, PipelinedFrameStateRHI.BeginCount);
 
 	// We need to ensure we release the swap chain images even if the session is not running.
 	if (PipelinedLayerStateRHI.ColorSwapchain)
@@ -3591,7 +3613,7 @@ void FOpenXRHMD::OnFinishRendering_RHIThread()
 #endif
 		static int EndCount = 0;
 		PipelinedFrameStateRHI.EndCount = ++EndCount;
-		UE_LOG(LogHMD, VeryVerbose, TEXT("xrEndFrame WaitCount: %i BeginCount: %i EndCount: %i"), PipelinedFrameStateRHI.WaitCount, PipelinedFrameStateRHI.BeginCount, PipelinedFrameStateRHI.EndCount);
+		UE_LOG(LogHMD, VeryVerbose, TEXT("%s WF_%i xrEndFrame WaitCount: %i BeginCount: %i EndCount: %i"), HMDThreadString(), PipelinedFrameStateRHI.WaitCount, PipelinedFrameStateRHI.WaitCount, PipelinedFrameStateRHI.BeginCount, PipelinedFrameStateRHI.EndCount);
 		XR_ENSURE(xrEndFrame(Session, &EndInfo));
 	}
 
diff --git a/Engine/Plugins/Runtime/OpenXR/Source/OpenXRHMD/Private/OpenXRHMD_Swapchain.cpp b/Engine/Plugins/Runtime/OpenXR/Source/OpenXRHMD/Private/OpenXRHMD_Swapchain.cpp
--- a/Engine/Plugins/Runtime/OpenXR/Source/OpenXRHMD/Private/OpenXRHMD_Swapchain.cpp
+++ b/Engine/Plugins/Runtime/OpenXR/Source/OpenXRHMD/Private/OpenXRHMD_Swapchain.cpp
@@ -43,10 +43,10 @@ void FOpenXRSwapchain::IncrementSwapChainIndex_RHIThread()
 	Info.next = nullptr;
 	XR_ENSURE(xrAcquireSwapchainImage(Handle, &Info, &SwapChainIndex));
 
-	UE_LOG(LogHMD, VeryVerbose, TEXT("Acquired image %d in swapchain %p"), SwapChainIndex, reinterpret_cast<const void*>(Handle));
-
 	RHITexture = RHITextureSwapChain[SwapChainIndex];
 	SwapChainIndex_RHIThread = SwapChainIndex;
+	
+	UE_LOG(LogHMD, VeryVerbose, TEXT("FOpenXRSwapchain::IncrementSwapChainIndex_RHIThread() Acquired image %d in swapchain %p metal texture: 0x%x"), SwapChainIndex, reinterpret_cast<const void*>(Handle), RHITexture.GetReference()->GetNativeResource());
 }
 
 void FOpenXRSwapchain::WaitCurrentImage_RHIThread(int64 Timeout)
@@ -92,8 +92,8 @@ void FOpenXRSwapchain::WaitCurrentImage_RHIThread(int64 Timeout)
 		// We can't continue without acquiring a new swapchain image since we won't have an image available to render to.
 		UE_LOG(LogHMD, Fatal, TEXT("Failed to wait on acquired swapchain image. This usually indicates a problem with the OpenXR runtime."));
 	}
-
-	UE_LOG(LogHMD, VeryVerbose, TEXT("Waited on image %d in swapchain %p"), SwapChainIndex_RHIThread.load(), reinterpret_cast<const void*>(Handle));
+	
+	UE_LOG(LogHMD, VeryVerbose, TEXT("FOpenXRSwapchain::WaitCurrentImage_RHIThread() Waited on image swapchain %p"), reinterpret_cast<const void*>(Handle));
 }
 
 void FOpenXRSwapchain::ReleaseCurrentImage_RHIThread()
@@ -115,7 +115,7 @@ void FOpenXRSwapchain::ReleaseCurrentImage_RHIThread()
 	ReleaseInfo.next = nullptr;
 	XR_ENSURE(xrReleaseSwapchainImage(Handle, &ReleaseInfo));
 
-	UE_LOG(LogHMD, VeryVerbose, TEXT("Released on image %d in swapchain %p"), SwapChainIndex_RHIThread.load(), reinterpret_cast<const void*>(Handle));
+	UE_LOG(LogHMD, VeryVerbose, TEXT("FOpenXRSwapchain::ReleaseCurrentImage_RHIThread() Released on image in swapchain %p"), reinterpret_cast<const void*>(Handle));
 }
 
 uint8 FOpenXRSwapchain::GetNearestSupportedSwapchainFormat(XrSession InSession, uint8 RequestedFormat, TFunction<uint32(uint8)> ToPlatformFormat /*= nullptr*/)
diff --git a/Engine/Source/Runtime/Apple/MetalRHI/Private/MetalProfiler.cpp b/Engine/Source/Runtime/Apple/MetalRHI/Private/MetalProfiler.cpp
--- a/Engine/Source/Runtime/Apple/MetalRHI/Private/MetalProfiler.cpp
+++ b/Engine/Source/Runtime/Apple/MetalRHI/Private/MetalProfiler.cpp
@@ -19,7 +19,7 @@ DEFINE_STAT(STAT_MetalTextureMemUpdate);
 DEFINE_STAT(STAT_MetalDrawCallTime);
 DEFINE_STAT(STAT_MetalPipelineStateTime);
 DEFINE_STAT(STAT_MetalPrepareDrawTime);
-
+DEFINE_STAT(STAT_MetalSwitchToNoneTime);
 DEFINE_STAT(STAT_MetalSwitchToRenderTime);
 DEFINE_STAT(STAT_MetalSwitchToComputeTime);
 DEFINE_STAT(STAT_MetalSwitchToBlitTime);
diff --git a/Engine/Source/Runtime/Apple/MetalRHI/Private/MetalProfiler.h b/Engine/Source/Runtime/Apple/MetalRHI/Private/MetalProfiler.h
--- a/Engine/Source/Runtime/Apple/MetalRHI/Private/MetalProfiler.h
+++ b/Engine/Source/Runtime/Apple/MetalRHI/Private/MetalProfiler.h
@@ -13,6 +13,7 @@ class FMetalContext;
 DECLARE_CYCLE_STAT_EXTERN(TEXT("MakeDrawable time"),STAT_MetalMakeDrawableTime,STATGROUP_MetalRHI, );
 DECLARE_CYCLE_STAT_EXTERN(TEXT("Draw call time"),STAT_MetalDrawCallTime,STATGROUP_MetalRHI, );
 DECLARE_CYCLE_STAT_EXTERN(TEXT("PrepareDraw time"),STAT_MetalPrepareDrawTime,STATGROUP_MetalRHI, );
+DECLARE_CYCLE_STAT_EXTERN(TEXT("SwitchToNone time"),STAT_MetalSwitchToNoneTime,STATGROUP_MetalRHI, );
 DECLARE_CYCLE_STAT_EXTERN(TEXT("SwitchToRender time"),STAT_MetalSwitchToRenderTime,STATGROUP_MetalRHI, );
 DECLARE_CYCLE_STAT_EXTERN(TEXT("SwitchToCompute time"),STAT_MetalSwitchToComputeTime,STATGROUP_MetalRHI, );
 DECLARE_CYCLE_STAT_EXTERN(TEXT("SwitchToBlit time"),STAT_MetalSwitchToBlitTime,STATGROUP_MetalRHI, );
diff --git a/Engine/Source/Runtime/Apple/MetalRHI/Private/MetalRHIVisionOSBridge.cpp b/Engine/Source/Runtime/Apple/MetalRHI/Private/MetalRHIVisionOSBridge.cpp
--- a/Engine/Source/Runtime/Apple/MetalRHI/Private/MetalRHIVisionOSBridge.cpp
+++ b/Engine/Source/Runtime/Apple/MetalRHI/Private/MetalRHIVisionOSBridge.cpp
@@ -16,8 +16,7 @@ void MetalRHIVisionOS::PresentImmersive(const MetalRHIVisionOS::PresentImmersive
     check(RHICommandContext->CustomPresentViewport);
     FMetalViewport* Viewport = ResourceCast(RHICommandContext->CustomPresentViewport);
     check(Viewport);
-    FMetalSurface* Surface = GetMetalSurfaceFromRHITexture(Params.Texture);
 
-    Viewport->PresentImmersive(Surface, Params);
+    Viewport->PresentImmersive(Params);
 }
 #endif // PLATFORM_VISIONOS
diff --git a/Engine/Source/Runtime/Apple/MetalRHI/Private/MetalRenderPass.cpp b/Engine/Source/Runtime/Apple/MetalRHI/Private/MetalRenderPass.cpp
--- a/Engine/Source/Runtime/Apple/MetalRHI/Private/MetalRenderPass.cpp
+++ b/Engine/Source/Runtime/Apple/MetalRHI/Private/MetalRenderPass.cpp
@@ -95,6 +95,14 @@ TRefCountPtr<FMetalFence> const& FMetalRenderPass::Submit(EMetalSubmitFlags Flag
         CurrentEncoder.CommitCommandBuffer(Flags);
     }
 	
+#if PLATFORM_VISIONOS
+	if (CompositorServicesFrame)
+	{
+		cp_frame_end_submission(CompositorServicesFrame);
+		CompositorServicesFrame = nullptr;
+	}
+#endif
+
 	OutstandingBufferUploads.Empty();
 	if (Flags & EMetalSubmitFlagsResetState)
 	{
@@ -838,6 +846,19 @@ void FMetalRenderPass::PresentTexture(MTL::Texture* Texture, uint32 sourceSlice,
     Encoder->copyFromTexture(Texture, sourceSlice, sourceLevel, sourceOrigin, sourceSize, toTexture, destinationSlice, destinationLevel, destinationOrigin);
 }
 
+#if PLATFORM_VISIONOS
+void FMetalRenderPass::EncodePresentImmersive(cp_drawable_t Drawable, cp_frame_t Frame)
+{
+	check(bWithinRenderPass == false); // Call RenderPassEnd before calling this function
+	
+	FMetalCommandBuffer* CurrentCommandBuffer = GetCurrentCommandBuffer();
+	check(CurrentCommandBuffer);
+	check(CompositorServicesFrame == nullptr);
+	cp_drawable_encode_present(Drawable, (__bridge id<MTLCommandBuffer>)CurrentCommandBuffer->GetMTLCmdBuffer().get());
+	CompositorServicesFrame = Frame;
+}
+#endif
+
 void FMetalRenderPass::SynchronizeTexture(MTL::Texture* Texture, uint32 Slice, uint32 Level)
 {
 	check(Texture);
diff --git a/Engine/Source/Runtime/Apple/MetalRHI/Private/MetalRenderPass.h b/Engine/Source/Runtime/Apple/MetalRHI/Private/MetalRenderPass.h
--- a/Engine/Source/Runtime/Apple/MetalRHI/Private/MetalRenderPass.h
+++ b/Engine/Source/Runtime/Apple/MetalRHI/Private/MetalRenderPass.h
@@ -67,6 +67,10 @@ public:
 	
 	void PresentTexture(MTL::Texture* Texture, uint32 sourceSlice, uint32 sourceLevel, MTL::Origin sourceOrigin, MTL::Size sourceSize, MTL::Texture* toTexture, uint32 destinationSlice, uint32 destinationLevel, MTL::Origin destinationOrigin);
     
+#if PLATFORM_VISIONOS
+	void EncodePresentImmersive(cp_drawable_t Drawable, cp_frame_t Frame);
+#endif
+	
     void SynchronizeTexture(MTL::Texture* Texture, uint32 Slice, uint32 Level);
     
 	void SynchroniseResource(MTL::Resource* Resource);
@@ -193,4 +197,8 @@ private:
 	MTL::DispatchType ComputeDispatchType;
     uint32 NumOutstandingOps;
     bool bWithinRenderPass;
+	
+#if PLATFORM_VISIONOS	
+	cp_frame_t CompositorServicesFrame;
+#endif
 };
diff --git a/Engine/Source/Runtime/Apple/MetalRHI/Private/MetalViewport.cpp b/Engine/Source/Runtime/Apple/MetalRHI/Private/MetalViewport.cpp
--- a/Engine/Source/Runtime/Apple/MetalRHI/Private/MetalViewport.cpp
+++ b/Engine/Source/Runtime/Apple/MetalRHI/Private/MetalViewport.cpp
@@ -608,46 +608,37 @@ void FMetalViewport::Swap()
 }
 
 #if PLATFORM_VISIONOS
-void FMetalViewport::GetDrawableImmersiveTextures(EMetalViewportAccessFlag Accessor, cp_frame_t SwiftLayerFrame, MTL::Texture*& OutColorTexture, MTL::Texture*& OutDepthTexture)
+void FMetalViewport::GetDrawableImmersiveTextures(EMetalViewportAccessFlag Accessor, cp_drawable_t SwiftDrawable, MTL::Texture*& OutColorTexture, MTL::Texture*& OutDepthTexture)
 {
-	check(SwiftLayerFrame != nullptr);
-
-	// get the next SwiftDrawable
-	UE_LOG(LogMetalVisionOS, Verbose, TEXT("SwiftLayerFrame(0x%x) cp_frame_query_drawable in FMetalViewport::GetDrawableTexture"), SwiftLayerFrame);
-	cp_drawable_t SwiftDrawable = cp_frame_query_drawable(SwiftLayerFrame);
-	if (SwiftDrawable == nullptr)
-	{
-		UE_LOG(LogMetal, Fatal, TEXT("Failed to get next SwiftDrawable"));
-	}
+	check(SwiftDrawable != nullptr);
 	
 	// get the color texture out and use that with the RHI
 	uint32 Index = GetViewportIndex(Accessor);
 	uint32 TextureCount = cp_drawable_get_texture_count(SwiftDrawable);
-	UE_LOG(LogMetalVisionOS, Verbose, TEXT("cp_drawable_get_texture_count %i"), TextureCount);
+	check(TextureCount = 1);
 	OutColorTexture = (__bridge MTL::Texture*)cp_drawable_get_color_texture(SwiftDrawable, 0);
 	OutDepthTexture = (__bridge MTL::Texture*)cp_drawable_get_depth_texture(SwiftDrawable, 0);
 	DrawableTextures[Index] = OutColorTexture;
 }
 
 // This is the present for Immersive visionOS, through the OXRVisionOS plugin.
-void FMetalViewport::PresentImmersive(TRefCountPtr<FMetalSurface> CompleteFrame, const MetalRHIVisionOS::PresentImmersiveParams& VisionOSParams)
+void FMetalViewport::PresentImmersive(const MetalRHIVisionOS::PresentImmersiveParams& VisionOSParams)
 {
-	UE_LOG(LogMetalVisionOS, Verbose, TEXT("FMetalViewport::PresentImmersive"));
 	check(SwiftLayer);  // If no SwiftLayer we should not be trying to be immersive.
+	check(VisionOSParams.SwiftFrame);
 
 	FMetalRHICommandContext* Context = static_cast<FMetalRHICommandContext*>(RHIGetDefaultContext());
 	FMetalDeviceContext& DeviceContext = GetMetalDeviceContext();
-	FMetalCommandQueue& CommandQueue = DeviceContext.GetCommandQueue();
-	
-	FScopeLock Lock(&Mutex);
-
-	LastCompleteFrame = CompleteFrame;
-	TRefCountPtr<FMetalSurface> LastCompleteDepth = GetMetalSurfaceFromRHITexture(VisionOSParams.Depth);
+	FMetalRenderPass& RenderPass = DeviceContext.GetCurrentRenderPass();
 
+	FScopeLock Lock(&Mutex);
+	
+	TRefCountPtr<FMetalSurface> MyLastCompleteFrame = GetMetalSurfaceFromRHITexture(VisionOSParams.Texture);
+	TRefCountPtr<FMetalSurface> MyLastCompleteDepth = GetMetalSurfaceFromRHITexture(VisionOSParams.Depth);
 	{
 		MTL::Texture* DrawableTextureParam = nullptr;
 		MTL::Texture* DrawableDepthTextureParam = nullptr;
-		GetDrawableImmersiveTextures(EMetalViewportAccessDisplayLink, VisionOSParams.SwiftFrame, DrawableTextureParam, DrawableDepthTextureParam);
+		GetDrawableImmersiveTextures(EMetalViewportAccessDisplayLink, VisionOSParams.SwiftDrawable, DrawableTextureParam, DrawableDepthTextureParam);
 		MTLTexturePtr DrawableTexture = NS::RetainPtr(DrawableTextureParam);
 		MTLTexturePtr DrawableDepthTexture = NS::RetainPtr(DrawableDepthTextureParam);
 		{
@@ -655,21 +646,13 @@ void FMetalViewport::PresentImmersive(TRefCountPtr<FMetalSurface> CompleteFrame,
 			
 			if (DrawableTexture)
 			{
-				FMetalCommandBuffer* CurrentCommandBuffer = CommandQueue.CreateCommandBuffer();
-				check(CurrentCommandBuffer);
-				
-#if ENABLE_METAL_GPUPROFILE
-				FMetalProfiler* Profiler = FMetalProfiler::GetProfiler();
-				FMetalCommandBufferStats* Stats = Profiler->AllocateCommandBuffer(CurrentCommandBuffer->GetMTLCmdBuffer(), 0);
-#endif
-
 				// TODO Currently we are using intermediate back buffer to connect the OXRVisionOS Swapchain to the drawable.
 				// I think we could use the drawable directly and avoid this copy.
 				check(GMetalSupportsIntermediateBackBuffer);
 				if (GMetalSupportsIntermediateBackBuffer)
 				{
 					{
-						TRefCountPtr<FMetalSurface> Texture = LastCompleteFrame;
+						TRefCountPtr<FMetalSurface> Texture = MyLastCompleteFrame;
 						check(IsValidRef(Texture));
 						MTLTexturePtr Src = Texture->Texture;
 						MTLTexturePtr& Dst = DrawableTexture;
@@ -677,24 +660,23 @@ void FMetalViewport::PresentImmersive(TRefCountPtr<FMetalSurface> CompleteFrame,
 						NSUInteger Width = FMath::Min(Src->width(), Dst->width());
 						NSUInteger Height = FMath::Min(Src->height(), Dst->height());
 						
-						MTL::BlitCommandEncoder* Encoder = CurrentCommandBuffer->GetMTLCmdBuffer()->blitCommandEncoder();
-						check(Encoder);
-						METAL_GPUPROFILE(Profiler->EncodeBlit(Stats, __FUNCTION__));
-						Encoder->copyFromTexture(Src.get(), 0, 0, MTL::Origin(0, 0, 0), MTL::Size(Width, Height, 1), Dst.get(), 0, 0, MTL::Origin(0, 0, 0));
-						Encoder->endEncoding();
+						RenderPass.CopyFromTextureToTexture(Src.get(), 0, 0, MTL::Origin(0, 0, 0), MTL::Size(Width, Height, 1), Dst.get(), 0, 0, MTL::Origin(0, 0, 0));
 					}
 
 					{
-//HACK BEGIN
-// using LastCompleteFrame, the color result, is very wrong, but the depth texture is not yet plumbed through
-// on the mobile rendere and providing nothing results in an all black scene while this results in mostly-ok-pixels.
-// With the deferred renderer depth does already work, so flip this to using depth.
-#if 1
-						TRefCountPtr<FMetalSurface> Texture = LastCompleteFrame;
-#else
-						TRefCountPtr<FMetalSurface> Texture = LastCompleteDepth;
-#endif
+						//HACK BEGIN
+						// using LastCompleteFrame, the color result, is very wrong, but the depth texture is not yet plumbed through
+						// on the mobile rendere and providing nothing results in an all black scene while this results in mostly-ok-pixels.
+						// With the deferred renderer depth does already work, so flip this to using depth.
+						static bool bUseRealDepth = false;
+						static bool bCheckedConfig = false;
+						if (bCheckedConfig == false) {
+							GConfig->GetBool(TEXT("/Script/IOSRuntimeSettings.IOSRuntimeSettings"), TEXT("bSupportsMetalMRT"), bUseRealDepth, GEngineIni);
+							bCheckedConfig = true;
+						}
+						TRefCountPtr<FMetalSurface> Texture = bUseRealDepth ? MyLastCompleteDepth : MyLastCompleteFrame;
 						//HACK END
+						
 						check(IsValidRef(Texture));
 						MTLTexturePtr Src = Texture->Texture;
 						MTLTexturePtr& Dst = DrawableDepthTexture;
@@ -702,44 +684,17 @@ void FMetalViewport::PresentImmersive(TRefCountPtr<FMetalSurface> CompleteFrame,
 						NS::UInteger Width = FMath::Min(Src->width(), Dst->width());
 						NS::UInteger Height = FMath::Min(Src->height(), Dst->height());
 						
-						MTLBlitCommandEncoderPtr Encoder = NS::RetainPtr(CurrentCommandBuffer->GetMTLCmdBuffer()->blitCommandEncoder());
-						check(Encoder);
-						METAL_GPUPROFILE(Profiler->EncodeBlit(Stats, __FUNCTION__));
-						Encoder->copyFromTexture(Src.get(), 0, 0, MTL::Origin(0, 0, 0), MTL::Size(Width, Height, 1), Dst.get(), 0, 0, MTL::Origin(0, 0, 0));
-						Encoder->endEncoding();
+						RenderPass.CopyFromTextureToTexture(Src.get(), 0, 0, MTL::Origin(0, 0, 0), MTL::Size(Width, Height, 1), Dst.get(), 0, 0, MTL::Origin(0, 0, 0));
 					}
 				}
-				
-				// This is a bit different than the usual pattern.
-				// This command buffer here is committed directly, instead of going through
-				// FMetalCommandList::Commit. So long as Present() is called within
-				// high level RHI BeginFrame/EndFrame this will not fine.
-				// Otherwise the recording of the Present time will be offset by one in the
-				// FMetalGPUProfiler frame indices.
-
-				MTL::HandlerFunction CommandBufferHandler = [](MTL::CommandBuffer* CommandBuffer)
-				{
-					FMetalGPUProfiler::RecordPresent(CommandBuffer);
-				};
-
-				if (VisionOSParams.SwiftFrame != nullptr) // TODO maybe remove this?
-				{
-					UE_LOG(LogMetalVisionOS, Verbose, TEXT("SwiftDrawable(0x%x) cp_drawable_encode_present in FMetalViewport::PresentImmersive"), VisionOSParams.SwiftDrawable);
-					cp_drawable_encode_present(VisionOSParams.SwiftDrawable, (__bridge id<MTLCommandBuffer>)CurrentCommandBuffer->GetMTLCmdBuffer().get());
-				}
-				
-				CurrentCommandBuffer->GetMTLCmdBuffer()->addCompletedHandler(CommandBufferHandler);
 
-				METAL_GPUPROFILE(Stats->End(CurrentCommandBuffer->GetMTLCmdBuffer()));
-				CommandQueue.CommitCommandBuffer(CurrentCommandBuffer);
+				RenderPass.EndRenderPass();
+				RenderPass.EncodePresentImmersive(VisionOSParams.SwiftDrawable, VisionOSParams.SwiftFrame);
 			}
 		}
 	}
-	
-	UE_LOG(LogMetalVisionOS, Verbose, TEXT("FMetalViewport::PresentImmersive calling cp_frame_end_submission with VisionOSParams.SwiftFrame(0x%x)"), VisionOSParams.SwiftFrame);
-	cp_frame_end_submission(VisionOSParams.SwiftFrame);
 }
-#endif
+#endif //PLATFORM_VISIONOS
 
 
 /*=============================================================================
diff --git a/Engine/Source/Runtime/Apple/MetalRHI/Public/MetalRHIVisionOSBridge.h b/Engine/Source/Runtime/Apple/MetalRHI/Public/MetalRHIVisionOSBridge.h
--- a/Engine/Source/Runtime/Apple/MetalRHI/Public/MetalRHIVisionOSBridge.h
+++ b/Engine/Source/Runtime/Apple/MetalRHI/Public/MetalRHIVisionOSBridge.h
@@ -18,6 +18,7 @@ namespace MetalRHIVisionOS
 		const FTextureRHIRef& Depth;
 		cp_frame_t& SwiftFrame;
         cp_drawable_t& SwiftDrawable;
+		int FrameCounter;
     };
     METALRHI_API void PresentImmersive(const PresentImmersiveParams& Params);
 }
diff --git a/Engine/Source/Runtime/Apple/MetalRHI/Public/MetalViewport.h b/Engine/Source/Runtime/Apple/MetalRHI/Public/MetalViewport.h
--- a/Engine/Source/Runtime/Apple/MetalRHI/Public/MetalViewport.h
+++ b/Engine/Source/Runtime/Apple/MetalRHI/Public/MetalViewport.h
@@ -78,8 +78,8 @@ public:
 	void Swap();
 	
 #if PLATFORM_VISIONOS
-	void GetDrawableImmersiveTextures(EMetalViewportAccessFlag Accessor, cp_frame_t SwiftLayerFrame, MTL::Texture*& OutColorTexture, MTL::Texture*& OutDepthTexture );
-    void PresentImmersive(TRefCountPtr<FMetalSurface> CompleteFrame, const MetalRHIVisionOS::PresentImmersiveParams& Params);
+	void GetDrawableImmersiveTextures(EMetalViewportAccessFlag Accessor, cp_drawable_t SwiftDrawable, MTL::Texture*& OutColorTexture, MTL::Texture*& OutDepthTexture );
+    void PresentImmersive(const MetalRHIVisionOS::PresentImmersiveParams& Params);
 #endif
 	
 private:
diff --git a/Engine/Source/Runtime/Launch/Private/IOS/LaunchIOS.cpp b/Engine/Source/Runtime/Launch/Private/IOS/LaunchIOS.cpp
--- a/Engine/Source/Runtime/Launch/Private/IOS/LaunchIOS.cpp
+++ b/Engine/Source/Runtime/Launch/Private/IOS/LaunchIOS.cpp
@@ -597,13 +597,13 @@ void FSwiftAppBootstrap::KickoffWithCompositingLayer(CP_OBJECT_cp_layer_renderer
 	
 	cp_layer_renderer_properties_t Props = cp_layer_renderer_get_properties(Layer);
 	int NumViews = cp_layer_renderer_properties_get_view_count(Props);
-	int NumTopologies = cp_layer_renderer_properties_get_texture_topology_count(Props);
 
 	NSMutableArray* Viewports = [NSMutableArray arrayWithCapacity:NumViews];
 
 	// get the texture topology
 	// @todo when Apple adds the API to actually get the size, use this instead of the mess below (docs indicate you can get
 	// get the width/height, but there's no functions to get them 
+	//int NumTopologies = cp_layer_renderer_properties_get_texture_topology_count(Props);
 //	for (int TopoIndex = 0; TopoIndex < NumToplogies; TopoIndex++)
 //	{
 //		cp_texture_topology_t Topology = cp_layer_renderer_properties_get_texture_topology(Props, TopoIndex);
