UE-210851 VisionOS build break
-Missing files from cl 32525048.

diff --git a/Engine/Platforms/VisionOS/Plugins/Runtime/OpenXRVisionOS/Source/OXRVisionOS/Private/OXRVisionOSSwapchain.cpp b/Engine/Platforms/VisionOS/Plugins/Runtime/OpenXRVisionOS/Source/OXRVisionOS/Private/OXRVisionOSSwapchain.cpp
--- a/Engine/Platforms/VisionOS/Plugins/Runtime/OpenXRVisionOS/Source/OXRVisionOS/Private/OXRVisionOSSwapchain.cpp
+++ b/Engine/Platforms/VisionOS/Plugins/Runtime/OpenXRVisionOS/Source/OXRVisionOS/Private/OXRVisionOSSwapchain.cpp
@@ -165,6 +165,18 @@ XrResult FOXRVisionOSSwapchain::XrEnumerateSwapchainImages(
 	return XrResult::XR_SUCCESS;
 }
 
+const FOXRVisionOSSwapchain::FSwapchainImage& FOXRVisionOSSwapchain::GetLastWaitedImage() const
+{
+	check(NextImageToWait < Images.Num());
+
+	const FSwapchainImage& WaitedImage = Images[NextImageToWait];
+	//check(WaitedImage.ImageState == EImageState::WaitComplete);
+
+	UE_LOG(LogOXRVisionOS, VeryVerbose, TEXT("FOXRVisionOSSwapchain::GetLastWaitedImage() NextImageToWait=%i Metal texture = 0x%x"), NextImageToWait, WaitedImage.Image.GetReference()->GetNativeResource());
+	
+	return WaitedImage;
+}
+
 const FOXRVisionOSSwapchain::FSwapchainImage& FOXRVisionOSSwapchain::GetLastReleasedImage() const
 {
 	check(NextImageToRelease < Images.Num());
@@ -175,6 +187,8 @@ const FOXRVisionOSSwapchain::FSwapchainImage& FOXRVisionOSSwapchain::GetLastRele
 	const FSwapchainImage& ReleasedImage = Images[NextImageToRelease];
 	check(ReleasedImage.ImageState == EImageState::Released);
 
+	UE_LOG(LogOXRVisionOS, VeryVerbose, TEXT("FOXRVisionOSSwapchain::GetLastReleasedImage() NextImageToRelease=%i Metal texture = 0x%x"), NextImageToRelease, ReleasedImage.Image.GetReference()->GetNativeResource());
+	
 	return ReleasedImage;
 }
 
@@ -201,6 +215,8 @@ XrResult FOXRVisionOSSwapchain::XrAcquireSwapchainImage(
 	{
 		NextImage.ImageState = EImageState::Aquired;
 		*index = NextImageToAcquire;
+		
+		UE_LOG(LogOXRVisionOS, VeryVerbose, TEXT("FOXRVisionOSSwapchain::XrAcquireSwapchainImage() Acquired %i metal texture 0x%x"), NextImageToAcquire, NextImage.Image.GetReference()->GetNativeResource());
 
 		AcquiredImageQueue.Enqueue(NextImageToAcquire);
 
@@ -248,11 +264,13 @@ XrResult FOXRVisionOSSwapchain::XrWaitSwapchainImage(
 	FSwapchainImage& NextImage = Images[NextImageToWait];
 	if (NextImage.ImageState == EImageState::Aquired)
 	{
+		UE_LOG(LogOXRVisionOS, VeryVerbose, TEXT("FOXRVisionOSSwapchain::XrWaitSwapchainImage() Waiting on NextImageToWait %i "), NextImageToWait);
 		NextImage.ImageState = EImageState::Waiting;
 
 		//custom image ready wait could be here, if there was one.
 
 		NextImage.ImageState = EImageState::WaitComplete;
+		UE_LOG(LogOXRVisionOS, VeryVerbose, TEXT("FOXRVisionOSSwapchain::XrWaitSwapchainImage() WaitComplete on NextImageToWait %i  metal texture 0x%x"), NextImageToWait, NextImage.Image.GetReference()->GetNativeResource());
 		return XrResult::XR_SUCCESS;
 	}
 	else
@@ -279,7 +297,9 @@ XrResult FOXRVisionOSSwapchain::XrReleaseSwapchainImage(
 	FSwapchainImage& NextImage = Images[NextImageToRelease];
 	if (NextImage.ImageState == EImageState::WaitComplete)
 	{
-				NextImage.ImageState = EImageState::Released;
+		UE_LOG(LogOXRVisionOS, VeryVerbose, TEXT("FOXRVisionOSSwapchain::XrReleaseSwapchainImage() Releasing %i metal texture 0x%x"), NextImageToRelease, NextImage.Image.GetReference()->GetNativeResource());
+		
+		NextImage.ImageState = EImageState::Released;
 		OutstandingWaits--;
 		return XrResult::XR_SUCCESS;
 	}
diff --git a/Engine/Platforms/VisionOS/Plugins/Runtime/OpenXRVisionOS/Source/OXRVisionOS/Private/OXRVisionOSSwapchain.h b/Engine/Platforms/VisionOS/Plugins/Runtime/OpenXRVisionOS/Source/OXRVisionOS/Private/OXRVisionOSSwapchain.h
--- a/Engine/Platforms/VisionOS/Plugins/Runtime/OpenXRVisionOS/Source/OXRVisionOS/Private/OXRVisionOSSwapchain.h
+++ b/Engine/Platforms/VisionOS/Plugins/Runtime/OpenXRVisionOS/Source/OXRVisionOS/Private/OXRVisionOSSwapchain.h
@@ -54,6 +54,8 @@ public:
 
 	XrResult XrReleaseSwapchainImage(
 		const XrSwapchainImageReleaseInfo*			releaseInfo);
+	
+	const FSwapchainImage& GetLastWaitedImage() const;
 
 	const FSwapchainImage& GetLastReleasedImage() const;
 
