VisionOS hand tracking fix
-Fixed tracked flag bit flag setting error.

diff --git a/Engine/Platforms/VisionOS/Plugins/Runtime/OpenXRVisionOS/Source/OXRVisionOS/Private/OXRVisionOSSession.cpp b/Engine/Platforms/VisionOS/Plugins/Runtime/OpenXRVisionOS/Source/OXRVisionOS/Private/OXRVisionOSSession.cpp
--- a/Engine/Platforms/VisionOS/Plugins/Runtime/OpenXRVisionOS/Source/OXRVisionOS/Private/OXRVisionOSSession.cpp
+++ b/Engine/Platforms/VisionOS/Plugins/Runtime/OpenXRVisionOS/Source/OXRVisionOS/Private/OXRVisionOSSession.cpp
@@ -1562,7 +1562,7 @@ void FOXRVisionOSSession::SyncHandTracking()
 			else 
 			{
 				// Clear the tracked bits, but leave the valid bits and transforms as they are.  This means the joint is valid when it has been tracked once.
-				JointLocation.locationFlags = JointLocation.locationFlags && (XR_SPACE_LOCATION_POSITION_VALID_BIT | XR_SPACE_LOCATION_ORIENTATION_VALID_BIT);
+				JointLocation.locationFlags = JointLocation.locationFlags & (XR_SPACE_LOCATION_POSITION_VALID_BIT | XR_SPACE_LOCATION_ORIENTATION_VALID_BIT);
 			}
 			JointLocalTransform = ar_skeleton_joint_get_anchor_from_joint_transform(Joint); // Even untracked joints do get updates from their parent's positions.
 
diff --git a/Engine/Platforms/VisionOS/Plugins/Runtime/OpenXRVisionOS/Source/OXRVisionOS/Private/OXRVisionOSSession.h b/Engine/Platforms/VisionOS/Plugins/Runtime/OpenXRVisionOS/Source/OXRVisionOS/Private/OXRVisionOSSession.h
--- a/Engine/Platforms/VisionOS/Plugins/Runtime/OpenXRVisionOS/Source/OXRVisionOS/Private/OXRVisionOSSession.h
+++ b/Engine/Platforms/VisionOS/Plugins/Runtime/OpenXRVisionOS/Source/OXRVisionOS/Private/OXRVisionOSSession.h
@@ -287,6 +287,6 @@ private:
 	static_assert(XR_HAND_LEFT_EXT == 1 && XR_HAND_RIGHT_EXT == 2);  // We will assume this to use value-1 as array index.
 	FOXRVisionOSHandTracker& GetHandTracker(XrHandEXT Hand) { return HandTrackers[Hand - 1]; }
 	const FOXRVisionOSHandTracker& GetHandTracker(XrHandEXT Hand) const { return HandTrackers[Hand - 1]; }
-	bool ValidateXrHandTrackerEXT(XrHandTrackerEXT handTracker) { return (((FOXRVisionOSHandTracker*)handTracker) - ((FOXRVisionOSHandTracker*)&HandTrackers)) <= 1; }
+	bool ValidateXrHandTrackerEXT(XrHandTrackerEXT HandTracker) { return (((FOXRVisionOSHandTracker*)HandTracker) - ((FOXRVisionOSHandTracker*)&HandTrackers)) <= 1; }
 	void SyncHandTracking();
 };
