visionOS Pause State and visionOS 2.0 infrastructure
-VisionOS 1.x vs 2.x defines set so we can compile for both.
-MetalRHI refactor updates.
-Compositor pause state implemented (meaning you can take the device off your head and put it back and your app will continue running).
-Improved shutdown behavior (still crashes, but later).
-cp_layer_ functions protected by mutex locks.
-Tested with visionOS 2.0 beta 7 + Xcode 16 beta 6 and with visionOS 1.1 + Xcode 15.3.

diff --git a/Engine/Platforms/VisionOS/Config/VisionOS_SDK.json b/Engine/Platforms/VisionOS/Config/VisionOS_SDK.json
--- a/Engine/Platforms/VisionOS/Config/VisionOS_SDK.json
+++ b/Engine/Platforms/VisionOS/Config/VisionOS_SDK.json
@@ -1,4 +1,6 @@
 {
 	// inherit the IOS MinSoftwareVersion
-	"ParentSDKFile": "../../../Config/IOS/IOS_SDK.json"
+	"ParentSDKFile": "../../../Config/IOS/IOS_SDK.json",
+
+	"MinVersion": "15.3.0",
 }
diff --git a/Engine/Platforms/VisionOS/Plugins/Runtime/OpenXRVisionOS/Source/OXRVisionOS/Private/OXRVisionOSSession.cpp b/Engine/Platforms/VisionOS/Plugins/Runtime/OpenXRVisionOS/Source/OXRVisionOS/Private/OXRVisionOSSession.cpp
--- a/Engine/Platforms/VisionOS/Plugins/Runtime/OpenXRVisionOS/Source/OXRVisionOS/Private/OXRVisionOSSession.cpp
+++ b/Engine/Platforms/VisionOS/Plugins/Runtime/OpenXRVisionOS/Source/OXRVisionOS/Private/OXRVisionOSSession.cpp
@@ -125,6 +125,7 @@ void FOXRVisionOSSession::OnWorldTickStart(UWorld* World, ELevelTick TickType, f
 {
 	if (SessionState == XR_SESSION_STATE_IDLE)
 	{
+		FScopeLock Lock(&CpLayerMutex);
 		switch (cp_layer_renderer_get_state(OXRVisionOS::GetSwiftLayerRenderer()))
 		{
 			case cp_layer_renderer_state_paused:
@@ -827,44 +828,66 @@ XrResult FOXRVisionOSSession::XrWaitFrame(
 		return XrResult::XR_ERROR_SESSION_NOT_RUNNING;
 	}
 	
-	XrResult SuccessfulReturnValue =  XrResult::XR_SUCCESS;
+	// Block until the previous frame's xrBeginFrame has happened.
+	{
+		SCOPED_NAMED_EVENT_TEXT("FOXRVisionOSSession::XrWaitFrame XrWaitFrameEvent->Wait()", FColor::Turquoise);
+		UE_LOG(LogOXRVisionOS, Verbose, TEXT("XrWaitFrameEvent->Wait() Started   FC will be %i"), SessionFrameCounter);
+		XrWaitFrameEvent->Wait();
+	}
 	
-	switch (cp_layer_renderer_get_state(OXRVisionOS::GetSwiftLayerRenderer()))
+	XrResult SuccessfulReturnValue =  XrResult::XR_SUCCESS;
+	bool bSetSyncronizedThisFrame = false;
 	{
-		case cp_layer_renderer_state_paused:
-			//TODO: handle app suspending, this happens when the device locks.  
-			check(false);
-			break;
-
-		case cp_layer_renderer_state_running:
-			break;
-
-		case cp_layer_renderer_state_invalidated:
-			// Not entirely sure this is correct, but it's kind of close.
-			SetSessionState(XrSessionState::XR_SESSION_STATE_LOSS_PENDING);
-			SuccessfulReturnValue = XrResult::XR_SESSION_LOSS_PENDING;
-			break;
-		default:
-			check(false);
+		FScopeLock Lock(&CpLayerMutex);
+		switch (cp_layer_renderer_get_state(OXRVisionOS::GetSwiftLayerRenderer()))
+		{
+			case cp_layer_renderer_state_paused:
+				// Handle app pause, this happens when you take the device off your head.
+				// We want to call cp_layer_renderer_wait_until_running to allow the system to take over, however
+				// we need in flight frames to complete first so that we don't keep making api calls
+				// and we want to call that on the rendering thread because apple only allows these functions to
+				// be called from a single thread.
+				// So we will go to the synchronized state for this frame and call cp_layer_renderer_wait_until_running
+				// in XrEndFrame instead of doing a normal submission to the compositor.  That will block the render thread.
+				// When that blocking call returns our next new frame will transition back to VISIBILE then FOCUSED below.
+				UE_LOG(LogOXRVisionOS, Log, TEXT("FOXRVisionOSSession XrWaitFrame cp_layer_renderer_get_state returned cp_layer_renderer_state_paused.  Transitioning to XR_SESSION_STATE_SYNCHRONIZED, shouldRender will be false."));
+				if (GetSessionState() != XrSessionState::XR_SESSION_STATE_SYNCHRONIZED)
+				{
+					SetSessionState(XrSessionState::XR_SESSION_STATE_SYNCHRONIZED);
+				}
+				PipelinedFrameStateGame.bSynchronizing = true;
+				bSetSyncronizedThisFrame = true;
+				break;
+				
+			case cp_layer_renderer_state_running:
+				break;
+				
+			case cp_layer_renderer_state_invalidated:
+				// Not entirely sure this is correct, but it's kind of close.
+				UE_LOG(LogOXRVisionOS, Log, TEXT("FOXRVisionOSSession XrWaitFrame LayerRendererState == cp_layer_renderer_state_invalidated.  SetSessionState to XR_SESSION_STATE_LOSS_PENDING."));
+				SetSessionState(XrSessionState::XR_SESSION_STATE_LOSS_PENDING);
+				SuccessfulReturnValue = XrResult::XR_SESSION_LOSS_PENDING;
+				PipelinedFrameStateGame.bSessionLost = true;
+				break;
+			default:
+				check(false);
+		}
 	}
 	
 	if (GetSessionState() == XrSessionState::XR_SESSION_STATE_READY)
 	{
 		SetSessionState(XrSessionState::XR_SESSION_STATE_SYNCHRONIZED);
+	}
+	
+	if (GetSessionState() == XrSessionState::XR_SESSION_STATE_SYNCHRONIZED && !bSetSyncronizedThisFrame)
+	{
 		SetSessionState(XrSessionState::XR_SESSION_STATE_VISIBLE);
 		SetSessionState(XrSessionState::XR_SESSION_STATE_FOCUSED);
 
-		SessionFrameCounter += OXRVisionOSSessionHelpers::OXRVisionOSBackbufferLength;  // Jump ahead in case the current frame counter was partially used while the previous session was being ended.  
+		SessionFrameCounter += OXRVisionOSSessionHelpers::OXRVisionOSBackbufferLength;  // Jump ahead in case the current frame counter was partially used while the previous session was being ended.
 		PipelinedFrameStateGame.bSynchronizing = false;
 	}
 
-	// Block until the previous frame's xrBeginFrame has happened.
-	{
-		SCOPED_NAMED_EVENT_TEXT("FOXRVisionOSSession::XrWaitFrame XrWaitFrameEvent->Wait()", FColor::Turquoise);
-		UE_LOG(LogOXRVisionOS, Verbose, TEXT("XrWaitFrameEvent->Wait() Started   FC will be %i"), SessionFrameCounter);
-		XrWaitFrameEvent->Wait();
-	}
-
 	PipelinedFrameStateGame.FrameCounter = SessionFrameCounter++;
 	// Each frame we read this index on the game thread and then we write the index after this on the render thread.
 	PipelinedFrameStateGame.RenderToGameFrameStateIndex++;
@@ -885,6 +908,7 @@ XrResult FOXRVisionOSSession::XrWaitFrame(
     LocateViewInfoIndexAdvance();
 
 	frameState->shouldRender = SessionState == XR_SESSION_STATE_FOCUSED || SessionState == XR_SESSION_STATE_VISIBLE;
+	PipelinedFrameStateGame.bShouldRender = frameState->shouldRender;
 
 	// We have to pipeline state updates in order to prevent task overlap (esp during map loading where we kick extra renders)
 	ENQUEUE_RENDER_COMMAND(UpdatePipelinedFrameState)([this, GameFrameState = PipelinedFrameStateGame](FRHICommandListImmediate& RHICmdList)
@@ -895,33 +919,55 @@ XrResult FOXRVisionOSSession::XrWaitFrame(
 		FPipelinedFrameState& FrameState = PipelinedFrameStateRendering;
 		FrameState = GameFrameState;
 		
-		FrameState.SwiftFrame = cp_layer_renderer_query_next_frame(VOSLayerRenderer);
-		UE_LOG(LogOXRVisionOS, VeryVerbose, TEXT("%s FC_%i SLF_0x%x  FOXRVisionOSSession::XrWaitFrame cp_layer_renderer_query_next_frame SwiftLayerFrame fetched"), VOSThreadString(), FrameState.FrameCounter, FrameState.SwiftFrame);
-
-		// Fetch the predicted timing information.
-		FrameState.SwiftFrameTiming = cp_frame_predict_timing(FrameState.SwiftFrame);
-		if ( FrameState.SwiftFrameTiming == nullptr)
+		if (FrameState.bSynchronizing) 
 		{
-			//TODO ??? failed frame?
-			// This means the layer is not in the correct state  Perhaps we are either running too early or we need to end the session.
-			// we might want to return XR_ERROR_SESSION_NOT_RUNNING here, but only if we actually end the session first.
-			assert(false);
-			//return;
+			UE_LOG(LogOXRVisionOS, Verbose, TEXT("%s FC_%i  FOXRVisionOSSession::XrWaitFrame FrameState.bSynchronizing, not getting a frame."), VOSThreadString(), FrameState.FrameCounter);
+			return;
 		}
 		
-		CFTimeInterval TargetRenderTimeInterval = cp_time_to_cf_time_interval(cp_frame_timing_get_presentation_time(FrameState.SwiftFrameTiming));
+		if (PipelinedFrameStateGame.bSessionLost)
+		{
+			UE_LOG(LogOXRVisionOS, Verbose, TEXT("%s FC_%i  FOXRVisionOSSession::XrWaitFrame FrameState.bSessionLost, not getting a frame."), VOSThreadString(), FrameState.FrameCounter);
+			return;
+		}
+		
+		{
+			FScopeLock Lock(&CpLayerMutex);
+			FrameState.SwiftFrame = cp_layer_renderer_query_next_frame(VOSLayerRenderer);
+		}
+		UE_LOG(LogOXRVisionOS, VeryVerbose, TEXT("%s FC_%i SLF_0x%x  FOXRVisionOSSession::XrWaitFrame cp_layer_renderer_query_next_frame SwiftLayerFrame fetched"), VOSThreadString(), FrameState.FrameCounter, FrameState.SwiftFrame);
+
+		if (FrameState.SwiftFrame) 
+		{
+			// Fetch the predicted timing information.
+			FrameState.SwiftFrameTiming = cp_frame_predict_timing(FrameState.SwiftFrame);
+			if ( FrameState.SwiftFrameTiming == nullptr)
+			{
+			 //TODO ??? failed frame?
+			 // This means the layer is not in the correct state  Perhaps we are either running too early or we need to end the session.
+			 // we might want to return XR_ERROR_SESSION_NOT_RUNNING here, but only if we actually end the session first.
+			 assert(false);
+			 //return;
+			}
 
-		UE_LOG(LogOXRVisionOS, VeryVerbose, TEXT("%s FC_%i SLF_0x%x FOXRVisionOSSession::XrWaitFrame cp_frame_predict_timing gave xr TargetRenderTimeInterval=%f"),
-			   VOSThreadString(), FrameState.FrameCounter, FrameState.SwiftFrame,
-			   (double)TargetRenderTimeInterval);
+			CFTimeInterval TargetRenderTimeInterval = cp_time_to_cf_time_interval(cp_frame_timing_get_presentation_time(FrameState.SwiftFrameTiming));
 
-		// Mark the beginnng of our game thread update work.  All the stuff before rendering.
-		UE_LOG(LogOXRVisionOS, VeryVerbose, TEXT("%s FC_%i SLF_0x%x FOXRVisionOSSession::XrWaitFrame cp_frame_start_update"), VOSThreadString(), FrameState.FrameCounter, FrameState.SwiftFrame);
-		cp_frame_start_update(FrameState.SwiftFrame);
+			UE_LOG(LogOXRVisionOS, VeryVerbose, TEXT("%s FC_%i SLF_0x%x FOXRVisionOSSession::XrWaitFrame cp_frame_predict_timing gave xr TargetRenderTimeInterval=%f"),
+				VOSThreadString(), FrameState.FrameCounter, FrameState.SwiftFrame,
+				(double)TargetRenderTimeInterval);
 
-		FrameState.PredictedDisplayTime = OXRVisionOS::APITimeToXrTime(TargetRenderTimeInterval);
-		
-		cp_frame_end_update(FrameState.SwiftFrame);
+			// Mark the beginnng of our game thread update work.  All the stuff before rendering.
+			UE_LOG(LogOXRVisionOS, VeryVerbose, TEXT("%s FC_%i SLF_0x%x FOXRVisionOSSession::XrWaitFrame cp_frame_start_update"), VOSThreadString(), FrameState.FrameCounter, FrameState.SwiftFrame);
+			cp_frame_start_update(FrameState.SwiftFrame);
+
+			FrameState.PredictedDisplayTime = OXRVisionOS::APITimeToXrTime(TargetRenderTimeInterval);
+
+			cp_frame_end_update(FrameState.SwiftFrame);
+		}
+		else
+		{
+			UE_LOG(LogOXRVisionOS, Log, TEXT("%s FC_%i SLF_0x%x  FOXRVisionOSSession::XrWaitFrame cp_layer_renderer_query_next_frame got null frame, if the session is stopping this is ok."), VOSThreadString(), FrameState.FrameCounter, FrameState.SwiftFrame);
+		}
 	});
 
 	// Update the STAGE space
@@ -933,7 +979,7 @@ XrResult FOXRVisionOSSession::XrWaitFrame(
 void FOXRVisionOSSession::OnBeginRendering_GameThread()
 {
     // Mark the end of our game frame
-	// Currently we are just not tracking the update phase at all.
+	// Currently we are going through this in UpdatePipelinedFrameState, because all cp_ calls have to happen on the render thread.  Perhaps in the future we will be able to track the game thread work.
 //    if (PipelinedFrameStateGame.SwiftFrame)
 //    {
 //		UE_LOG(LogOXRVisionOS, VeryVerbose, TEXT("%s FC_%i SLF_0x%x cp_frame_end_update in FOXRVisionOSSession::OnBeginRendering_GameThread"), VOSThreadString(), PipelinedFrameStateGame.FrameCounter, PipelinedFrameStateGame.SwiftFrame);
@@ -993,7 +1039,7 @@ XrResult FOXRVisionOSSession::XrBeginFrame(
 	}
     CachedBeginFlipFrameCounter = PipelinedFrameStateRHI.FrameCounter;
 	
-	if (PipelinedFrameStateRHI.bSynchronizing)
+	if (PipelinedFrameStateRHI.bSynchronizing || PipelinedFrameStateRHI.bSessionLost)
 	{
 		XrWaitFrameEvent->Trigger();
 		return XrResult::XR_SUCCESS;
@@ -1024,59 +1070,73 @@ void FOXRVisionOSSession::WaitUntil()
 {
     FPipelinedFrameState& FrameState = GetPipelinedFrameStateForThread();
     
-	cp_time_t OptimalInputTime = cp_frame_timing_get_optimal_input_time(FrameState.SwiftFrameTiming);
-	CFTimeInterval TimeInterval = cp_time_to_cf_time_interval(OptimalInputTime);
-
-	UE_LOG(LogOXRVisionOS, VeryVerbose, TEXT("%s FC_%i SLF_0x%x SD_0x%8x FOXRVisionOSSession::WaitUntil cp_time_wait_until started waiting until OptimalInputTime %f"), VOSThreadString(), FrameState.FrameCounter, FrameState.SwiftFrame, FrameState.SwiftDrawable, TimeInterval);
-	
-    cp_time_wait_until(OptimalInputTime);
-	
-	UE_LOG(LogOXRVisionOS, VeryVerbose, TEXT("%s FC_%i SLF_0x%x SD_0x%8x FOXRVisionOSSession::WaitUntil cp_time_wait_until completed"), VOSThreadString(), FrameState.FrameCounter, FrameState.SwiftFrame, FrameState.SwiftDrawable);
+	if (FrameState.SwiftFrame)
+	{
+		cp_time_t OptimalInputTime = cp_frame_timing_get_optimal_input_time(FrameState.SwiftFrameTiming);
+		CFTimeInterval TimeInterval = cp_time_to_cf_time_interval(OptimalInputTime);
+		
+		UE_LOG(LogOXRVisionOS, VeryVerbose, TEXT("%s FC_%i SLF_0x%x SD_0x%8x FOXRVisionOSSession::WaitUntil cp_time_wait_until started waiting until OptimalInputTime %f"), VOSThreadString(), FrameState.FrameCounter, FrameState.SwiftFrame, FrameState.SwiftDrawable, TimeInterval);
+		
+		cp_time_wait_until(OptimalInputTime);
+		
+		UE_LOG(LogOXRVisionOS, VeryVerbose, TEXT("%s FC_%i SLF_0x%x SD_0x%8x FOXRVisionOSSession::WaitUntil cp_time_wait_until completed"), VOSThreadString(), FrameState.FrameCounter, FrameState.SwiftFrame, FrameState.SwiftDrawable);
+	}
+	else
+	{
+		UE_LOG(LogOXRVisionOS, Verbose, TEXT("%s FC_%i SLF_0x%x SD_0x%8x FOXRVisionOSSession::WaitUntil doing nothing because SwiftFrame is null"), VOSThreadString(), FrameState.FrameCounter, FrameState.SwiftFrame, FrameState.SwiftDrawable);
+	}
 }
 
 void FOXRVisionOSSession::StartSubmission()
 {
     FPipelinedFrameState& FrameState = GetPipelinedFrameStateForThread();
 	
-	UE_LOG(LogOXRVisionOS, VeryVerbose, TEXT("%s FC_%i SLF_0x%x SD_0x%8x FOXRVisionOSSession::StartSubmission cp_frame_start_submission"), VOSThreadString(), FrameState.FrameCounter, FrameState.SwiftFrame, FrameState.SwiftDrawable);
-    
-    cp_frame_start_submission(FrameState.SwiftFrame);
-
-    FrameState.SwiftDrawable = cp_frame_query_drawable(FrameState.SwiftFrame);
-    check(FrameState.SwiftDrawable);
-	
-	UE_LOG(LogOXRVisionOS, VeryVerbose, TEXT("%s FC_%i SLF_0x%x SD_0x%8x FOXRVisionOSSession::StartSubmission cp_frame_query_drawable in got the SD again."), VOSThreadString(), FrameState.FrameCounter, FrameState.SwiftFrame, FrameState.SwiftDrawable);
-
-    FrameState.SwiftFinalFrameTiming = cp_drawable_get_frame_timing(FrameState.SwiftDrawable);
-	FrameState.SwiftFinalFrameTimeInterval = cp_time_to_cf_time_interval(cp_frame_timing_get_presentation_time(FrameState.SwiftFinalFrameTiming));
-	FrameState.DeviceAnchor = ar_device_anchor_create();
-	
-	UE_LOG(LogOXRVisionOS, VeryVerbose, TEXT("%s FC_%i SLF_0x%x SD_0x%8x FOXRVisionOSSession::StartSubmission cp_drawable_get_frame_timing got SwiftFinalFrameTimeInterval %f"), VOSThreadString(), FrameState.FrameCounter, FrameState.SwiftFrame, FrameState.SwiftDrawable, (double)FrameState.SwiftFinalFrameTimeInterval);
-
-	uint64_t CurrentMachTime = mach_absolute_time();
-	uint64_t PredictedPresentationMachTime = cp_frame_timing_get_presentation_time(FrameState.SwiftFinalFrameTiming).cp_mach_abs_time;
-	UE_LOG(LogOXRVisionOS, VeryVerbose, TEXT("   FC_%i                   FOXRVisionOSSession::StartSubmission CurrentMachTime %lld PredictedPresentationMachTime %lld"), FrameState.FrameCounter,
-		   CurrentMachTime, PredictedPresentationMachTime);
-
-	check(ARKitWorldTrackingProvider);
-	
+	if (FrameState.SwiftFrame)
 	{
-		auto anchor_status = ar_world_tracking_provider_query_device_anchor_at_timestamp(ARKitWorldTrackingProvider, FrameState.SwiftFinalFrameTimeInterval, FrameState.DeviceAnchor);
-		if (anchor_status == ar_device_anchor_query_status_success)
-		{
-			cp_drawable_set_device_anchor(FrameState.SwiftDrawable, FrameState.DeviceAnchor);
+		UE_LOG(LogOXRVisionOS, VeryVerbose, TEXT("%s FC_%i SLF_0x%x SD_0x%8x FOXRVisionOSSession::StartSubmission cp_frame_start_submission"), VOSThreadString(), FrameState.FrameCounter, FrameState.SwiftFrame, FrameState.SwiftDrawable);
+		
+		cp_frame_start_submission(FrameState.SwiftFrame);
+		
+		FrameState.SwiftDrawable = cp_frame_query_drawable(FrameState.SwiftFrame);
+    	check(FrameState.SwiftDrawable);
+		
+		UE_LOG(LogOXRVisionOS, VeryVerbose, TEXT("%s FC_%i SLF_0x%x SD_0x%8x FOXRVisionOSSession::StartSubmission cp_frame_query_drawable in got the SD again."), VOSThreadString(), FrameState.FrameCounter, FrameState.SwiftFrame, FrameState.SwiftDrawable);
+		
+		FrameState.SwiftFinalFrameTiming = cp_drawable_get_frame_timing(FrameState.SwiftDrawable);
+		FrameState.SwiftFinalFrameTimeInterval = cp_time_to_cf_time_interval(cp_frame_timing_get_presentation_time(FrameState.SwiftFinalFrameTiming));
+		FrameState.DeviceAnchor = ar_device_anchor_create();
+		
+		UE_LOG(LogOXRVisionOS, VeryVerbose, TEXT("%s FC_%i SLF_0x%x SD_0x%8x FOXRVisionOSSession::StartSubmission cp_drawable_get_frame_timing got SwiftFinalFrameTimeInterval %f"), VOSThreadString(), FrameState.FrameCounter, FrameState.SwiftFrame, FrameState.SwiftDrawable, (double)FrameState.SwiftFinalFrameTimeInterval);
+		
+		uint64_t CurrentMachTime = mach_absolute_time();
+		uint64_t PredictedPresentationMachTime = cp_frame_timing_get_presentation_time(FrameState.SwiftFinalFrameTiming).cp_mach_abs_time;
+		UE_LOG(LogOXRVisionOS, VeryVerbose, TEXT("   FC_%i                   FOXRVisionOSSession::StartSubmission CurrentMachTime %lld PredictedPresentationMachTime %lld"), FrameState.FrameCounter,
+			   CurrentMachTime, PredictedPresentationMachTime);
+		
+		check(ARKitWorldTrackingProvider);
 		
-			// Get the HMD transform and cache it
-			FrameState.HeadTransform = ar_anchor_get_origin_from_anchor_transform(FrameState.DeviceAnchor);
-		}
-		else
 		{
-			UE_LOG(LogOXRVisionOS, VeryVerbose, TEXT("%s FC_%i SLF_0x%x SD_0x%8x  FOXRVisionOSSession::StartSubmission cp_drawable_set_device_anchor DeviceAnchor query failed, setting nullptr, no reprojection"), VOSThreadString(), FrameState.FrameCounter, FrameState.SwiftFrame, FrameState.SwiftDrawable);
-			
-			cp_drawable_set_device_anchor(FrameState.SwiftDrawable, nullptr);
-			// Leave FrameState.HeadTransform untouched, same as game thread. It may simply be identity.  But whatever it is we will go ahead and render with it.
+			auto anchor_status = ar_world_tracking_provider_query_device_anchor_at_timestamp(ARKitWorldTrackingProvider, FrameState.SwiftFinalFrameTimeInterval, FrameState.DeviceAnchor);
+			if (anchor_status == ar_device_anchor_query_status_success)
+			{
+				cp_drawable_set_device_anchor(FrameState.SwiftDrawable, FrameState.DeviceAnchor);
+				
+				// Get the HMD transform and cache it
+				FrameState.HeadTransform = ar_anchor_get_origin_from_anchor_transform(FrameState.DeviceAnchor);
+			}
+			else
+			{
+				UE_LOG(LogOXRVisionOS, VeryVerbose, TEXT("%s FC_%i SLF_0x%x SD_0x%8x  FOXRVisionOSSession::StartSubmission cp_drawable_set_device_anchor DeviceAnchor query failed, setting nullptr, no reprojection"), VOSThreadString(), FrameState.FrameCounter, FrameState.SwiftFrame, FrameState.SwiftDrawable);
+				
+				cp_drawable_set_device_anchor(FrameState.SwiftDrawable, nullptr);
+				// Leave FrameState.HeadTransform untouched, same as game thread. It may simply be identity.  But whatever it is we will go ahead and render with it.
+			}
 		}
 	}
+	else
+	{
+		UE_LOG(LogOXRVisionOS, Verbose, TEXT("%s FC_%i SLF_0x%x SD_0x%8x FOXRVisionOSSession::StartSubmission doing nothing because SwiftFrame is null"), VOSThreadString(), FrameState.FrameCounter, FrameState.SwiftFrame, FrameState.SwiftDrawable);
+	}
 	RenderToGameFrameStateWrite(FrameState);
 }
 
@@ -1100,7 +1160,29 @@ XrResult FOXRVisionOSSession::XrEndFrame(
 
 	if (PipelinedFrameStateRHI.bSynchronizing)
 	{
-		UE_LOG(LogOXRVisionOS, Verbose, TEXT("FOXRVisionOSSession XrEndFrame bSynchronizing=true"));
+		UE_LOG(LogOXRVisionOS, Log, TEXT("FOXRVisionOSSession XrEndFrame bSynchronizing=true.  Calling cp_layer_renderer_wait_until_running"));
+		
+		{
+			FScopeLock Lock(&CpLayerMutex);
+			cp_layer_renderer_wait_until_running(VOSLayerRenderer);
+		}
+		UE_LOG(LogOXRVisionOS, Log, TEXT("FOXRVisionOSSession XrEndFrame cp_layer_renderer_wait_until_running returned"));
+		
+		// We need to call present so that it can signal the FrameSemaphore.
+		MetalRHIVisionOS::PresentImmersive(nullptr);
+		return XrResult::XR_SUCCESS;
+	}
+	
+	if (PipelinedFrameStateRHI.bSessionLost)
+	{
+		// We need to call present so that it can signal the FrameSemaphore.
+		MetalRHIVisionOS::PresentImmersive(nullptr);
+		return XrResult::XR_SUCCESS;
+	}
+	
+	if (PipelinedFrameStateRHI.SwiftFrame == nullptr)
+	{
+		UE_LOG(LogOXRVisionOS, Verbose, TEXT("%s FC_%i SLF_0x%x SD_0x%8x FOXRVisionOSSession::XrEndFrame with null SwiftFrame, not submitting"), VOSThreadString(), PipelinedFrameStateRHI.FrameCounter, PipelinedFrameStateRHI.SwiftFrame, PipelinedFrameStateRHI.SwiftDrawable);
 		return XrResult::XR_SUCCESS;
 	}
 
@@ -1138,9 +1220,12 @@ XrResult FOXRVisionOSSession::XrEndFrame(
 
 		const FOXRVisionOSSwapchain& SwapchainDepthData0 = *(reinterpret_cast<FOXRVisionOSSwapchain*>(Depth0->subImage.swapchain));
 		const FOXRVisionOSSwapchain::FSwapchainImage& WaitedDepth = SwapchainDepthData0.GetLastWaitedImage();
+		
+		const XrRHIContextEPIC* RHIContextEPIC = OpenXR::FindChainedStructByType<XrRHIContextEPIC>(InFrameEndInfo, (XrStructureType)XR_TYPE_RHI_CONTEXT_EPIC);
+		check(RHIContextEPIC && RHIContextEPIC->RHIContext);
 
-        const MetalRHIVisionOS::PresentImmersiveParams Params{WaitedImage.Image, WaitedDepth.Image, PipelinedFrameStateRHI.SwiftFrame, PipelinedFrameStateRHI.SwiftDrawable, PipelinedFrameStateRHI.FrameCounter};
-        MetalRHIVisionOS::PresentImmersive(Params);
+		const MetalRHIVisionOS::PresentImmersiveParams Params{WaitedImage.Image, WaitedDepth.Image, PipelinedFrameStateRHI.SwiftFrame, PipelinedFrameStateRHI.SwiftDrawable, PipelinedFrameStateRHI.FrameCounter, RHIContextEPIC->RHIContext};
+        MetalRHIVisionOS::PresentImmersive(&Params);
     }
 
 	NextBackBufferIndex = (NextBackBufferIndex + 1) % OXRVisionOSSessionHelpers::OXRVisionOSBackbufferLength;
@@ -1184,12 +1269,25 @@ XrResult FOXRVisionOSSession::XrLocateViews(
 	
 	//UE_LOG(LogOXRVisionOS, VeryVerbose, TEXT("%s FC_%i SLF_0x%x SD_0x%8x FOXRVisionOSSession::XrLocateViews"), VOSThreadString(), FrameState.FrameCounter, FrameState.SwiftFrame, FrameState.SwiftDrawable);
 
+	bool bUsePreviousFrameViews = false;
+	FLocateViewInfo PreviousViewInfo;
+	if (IsInGameThread())
+	{
+		// In the game thread we use the previous render thread views
+		PreviousViewInfo = GetLocateViewInfo_GameThread();
+		bUsePreviousFrameViews = true;
+	}
+	else if (FrameState.SwiftDrawable == nullptr)
+	{
+		// If a frame is abandoned we use the previous render thread views
+		PreviousViewInfo = GetLocateViewInfo_RenderThread();
+		bUsePreviousFrameViews = true;
+	}
+	
 	uint32_t NumViews = 0;
-    if (IsInGameThread())
+    if (bUsePreviousFrameViews)
     {
-        // We will use the NumViews from a previous render frame.
-        const FLocateViewInfo& ViewInfoCache = GetLocateViewInfo_GameThread();
-        NumViews = ViewInfoCache.NumViews;
+        NumViews = PreviousViewInfo.NumViews;
     }
     else
     {
@@ -1218,20 +1316,21 @@ XrResult FOXRVisionOSSession::XrLocateViews(
         
         for (int Index = 0; Index < NumViews; ++Index)
         {
-            if (IsInGameThread())
+            if (bUsePreviousFrameViews)
             {
-                const FLocateViewInfo& ViewInfoCache = GetLocateViewInfo_GameThread();
-                Views[Index].pose = OXRVisionOS::ToXrPose(ViewInfoCache.ViewTransforms[Index]);
-                Views[Index].fov = ViewInfoCache.HmdFovs[Index];
+                Views[Index].pose = OXRVisionOS::ToXrPose(PreviousViewInfo.ViewTransforms[Index]);
+                Views[Index].fov = PreviousViewInfo.HmdFovs[Index];
             }
             else
             {
                 check(IsInRenderingThread());
                 check(FrameState.SwiftDrawable);
 
-                cp_view_t View = cp_drawable_get_view(FrameState.SwiftDrawable, Index);
+				cp_view_t View = cp_drawable_get_view(FrameState.SwiftDrawable, Index);
+				
+#if VISIONOS_MAJOR_VERSION == 1
                 simd_float4 Tangents = cp_view_get_tangents(View);
-                //simd_float2 DepthRange = cp_drawable_get_depth_range(FrameState.SwiftDrawable);  //TODO ??? do we need this?  There is an openxrextension for it, could use that.
+				//simd_float2 DepthRange = cp_drawable_get_depth_range(FrameState.SwiftDrawable);  //TODO ??? do we need this?  There is an openxrextension for it, could use that.
     //                                                                          DepthRange[1], /* nearZ */
     //                                                                          DepthRange[0], /* farZ */
     //                                                                          true); /* reverseZ */
@@ -1240,7 +1339,28 @@ XrResult FOXRVisionOSSession::XrLocateViews(
 				FrameState.HmdFovs[Index].angleRight    =  FMath::Atan(Tangents[1]);
                 FrameState.HmdFovs[Index].angleUp       =  FMath::Atan(Tangents[2]);
                 FrameState.HmdFovs[Index].angleDown     = -FMath::Atan(Tangents[3]);
-
+#elif VISIONOS_MAJOR_VERSION >= 2
+				simd_float4x4 ProjectionMatrix = cp_drawable_compute_projection(FrameState.SwiftDrawable, cp_axis_direction_convention_right_up_back, Index);
+				
+				// Extract the tangents from the projection matrix
+				const float InvLmR = -ProjectionMatrix.columns[0][0] / 2.0f;
+				const float InvBmT = -ProjectionMatrix.columns[1][1] / 2.0f; 
+				const float SumRL = ProjectionMatrix.columns[2][0] / -InvLmR;
+				const float SumTB = ProjectionMatrix.columns[2][1] / -InvBmT;
+				const float LMinusR = 1.0 / InvLmR;
+				const float BMinusT = 1.0 / InvBmT;
+				const float L = (SumRL + LMinusR) / 2.0f;
+				const float R = SumRL - L;
+				const float B = (SumTB + BMinusT) / 2.0f;
+				const float T = SumTB - B;
+				
+				FrameState.HmdFovs[Index].angleLeft     = FMath::Atan(L);
+				FrameState.HmdFovs[Index].angleRight    = FMath::Atan(R);
+				FrameState.HmdFovs[Index].angleUp       = FMath::Atan(T);
+				FrameState.HmdFovs[Index].angleDown     = FMath::Atan(B);
+#else
+	static_assert(false); // Something wrong with the VISIONOS_MAJOR_VERSION define
+#endif
                 // Hack: We only xrLocateViews for the HMD device, and cp_view_get_transform gives us hmd relative transforms, so we don't need to do any additional math here.
 				// However OpenXR spec allows one to request the views in any space, which would require us to do some transforms and inverse transforms.
 				// Perhaps OpenXRHMD ought to be getting the view poses in view space rather than HMD device space because we could detect that here easily???
diff --git a/Engine/Platforms/VisionOS/Plugins/Runtime/OpenXRVisionOS/Source/OXRVisionOS/Private/OXRVisionOSSession.h b/Engine/Platforms/VisionOS/Plugins/Runtime/OpenXRVisionOS/Source/OXRVisionOS/Private/OXRVisionOSSession.h
--- a/Engine/Platforms/VisionOS/Plugins/Runtime/OpenXRVisionOS/Source/OXRVisionOS/Private/OXRVisionOSSession.h
+++ b/Engine/Platforms/VisionOS/Plugins/Runtime/OpenXRVisionOS/Source/OXRVisionOS/Private/OXRVisionOSSession.h
@@ -167,6 +167,8 @@ public:
 		CFTimeInterval SwiftFinalFrameTimeInterval = 0;
         ar_device_anchor_t DeviceAnchor = nullptr;
 		bool bSynchronizing = true; // when true we are in the process of starting up a new session and the new frame state has not yet propagated.
+		bool bShouldRender = false;
+		bool bSessionLost = false;
 		XrFovf HmdFovs[2];
         int32 LocateViewInfoBufferIndex = 0;
 		int32 RenderToGameFrameStateIndex = -1;
@@ -253,6 +255,11 @@ private:
 	ar_world_tracking_provider_t ARKitWorldTrackingProvider = nullptr;
 	ar_device_anchor_t ARKitHMDAnchor = nullptr;
 	
+	// mutexes to protect apple api functions from concurrent calls
+	// We do not protect the early startup calls, before OXRVisionOSSession is involved.  They should be single threaded.
+	mutable FCriticalSection CpLayerMutex; // Protect all cp_label_ calls.
+	//mutable FCriticalSection CpFrameMutex; // Currently we are restricting cp_frame_ calls to the Render thread by not having an RHI thread, however we may want to change that in the future...
+	
 	// Hand Tracking
 public:
 	struct FOXRVisionOSHandTracker
diff --git a/Engine/Platforms/VisionOS/Source/Programs/UnrealBuildTool/UEBuildVisionOS.cs b/Engine/Platforms/VisionOS/Source/Programs/UnrealBuildTool/UEBuildVisionOS.cs
--- a/Engine/Platforms/VisionOS/Source/Programs/UnrealBuildTool/UEBuildVisionOS.cs
+++ b/Engine/Platforms/VisionOS/Source/Programs/UnrealBuildTool/UEBuildVisionOS.cs
@@ -20,6 +20,9 @@ namespace UnrealBuildTool
 
 	class VisionOSPlatform : IOSPlatform
 	{
+		// Cached VisionOS sdk version from the toolchain, which detects it.
+		private float SDKVersionFloat = 0.0f;
+
 		public VisionOSPlatform(UEBuildPlatformSDK InSDK, ILogger Logger)
 			: base(InSDK, UnrealTargetPlatform.VisionOS, Logger)
 		{
@@ -94,9 +97,17 @@ namespace UnrealBuildTool
 			base.SetUpEnvironment(Target, CompileEnvironment, LinkEnvironment);
 			CompileEnvironment.Definitions.Add("PLATFORM_VISIONOS=1");
 
+			if (SDKVersionFloat < 2.0)
+			{
+				CompileEnvironment.Definitions.Add("VISIONOS_MAJOR_VERSION=1");
+			}
+			else
+			{
+				CompileEnvironment.Definitions.Add("VISIONOS_MAJOR_VERSION=2");
+			}
+
 			// VisionOS uses only IOS header files, so use it's platform headers
 			CompileEnvironment.Definitions.Add("OVERRIDE_PLATFORM_HEADER_NAME=IOS");
-
 		}
 
 		/// <summary>
@@ -107,7 +118,9 @@ namespace UnrealBuildTool
 		public override UEToolChain CreateToolChain(ReadOnlyTargetRules Target)
 		{
 			VisionOSProjectSettings ProjectSettings = ((VisionOSPlatform)UEBuildPlatform.GetBuildPlatform(UnrealTargetPlatform.VisionOS)).ReadProjectSettings(Target.ProjectFile);
-			return new VisionOSToolChain(Target, ProjectSettings, Logger);
+			VisionOSToolChain NewToolChain = new VisionOSToolChain(Target, ProjectSettings, Logger);
+			SDKVersionFloat = NewToolChain.GetSDKVersionFloat();
+			return NewToolChain;
 		}
 
 		public override void Deploy(TargetReceipt Receipt)
diff --git a/Engine/Platforms/VisionOS/Source/Programs/UnrealBuildTool/VisionOSToolChain.cs b/Engine/Platforms/VisionOS/Source/Programs/UnrealBuildTool/VisionOSToolChain.cs
--- a/Engine/Platforms/VisionOS/Source/Programs/UnrealBuildTool/VisionOSToolChain.cs
+++ b/Engine/Platforms/VisionOS/Source/Programs/UnrealBuildTool/VisionOSToolChain.cs
@@ -13,6 +13,8 @@ namespace UnrealBuildTool
 
 	class VisionOSToolChain : IOSToolChain
 	{
+		private VisionOSToolChainSettings Settings => (VisionOSToolChainSettings)ToolChainSettings.Value;
+
 		public VisionOSToolChain(ReadOnlyTargetRules InTarget, VisionOSProjectSettings InProjectSettings, ILogger InLogger)
 			: base(InTarget, InProjectSettings, () => new VisionOSToolChainSettings(InLogger), ClangToolChainOptions.None, InLogger)
 		{
@@ -22,5 +24,10 @@ namespace UnrealBuildTool
 		{
 			return "";
 		}
+
+		public float GetSDKVersionFloat()
+		{
+			return Settings.SDKVersionFloat;
+		}
 	}
 }
diff --git a/Engine/Source/Runtime/Apple/MetalRHI/Private/MetalRHIVisionOSBridge.cpp b/Engine/Source/Runtime/Apple/MetalRHI/Private/MetalRHIVisionOSBridge.cpp
--- a/Engine/Source/Runtime/Apple/MetalRHI/Private/MetalRHIVisionOSBridge.cpp
+++ b/Engine/Source/Runtime/Apple/MetalRHI/Private/MetalRHIVisionOSBridge.cpp
@@ -9,7 +9,7 @@
 
 DEFINE_LOG_CATEGORY(LogMetalVisionOS);
 
-void MetalRHIVisionOS::PresentImmersive(const MetalRHIVisionOS::PresentImmersiveParams& Params)
+void MetalRHIVisionOS::PresentImmersive(const MetalRHIVisionOS::PresentImmersiveParams* Params)
 {
     FMetalRHICommandContext* RHICommandContext = static_cast<FMetalRHICommandContext*>(RHIGetDefaultContext());
     check(RHICommandContext);
diff --git a/Engine/Source/Runtime/Apple/MetalRHI/Private/MetalViewport.cpp b/Engine/Source/Runtime/Apple/MetalRHI/Private/MetalViewport.cpp
--- a/Engine/Source/Runtime/Apple/MetalRHI/Private/MetalViewport.cpp
+++ b/Engine/Source/Runtime/Apple/MetalRHI/Private/MetalViewport.cpp
@@ -629,13 +629,25 @@ void FMetalViewport::GetDrawableImmersiveTextures(EMetalViewportAccessFlag Acces
 }
 
 // This is the present for Immersive visionOS, through the OXRVisionOS plugin.
-void FMetalViewport::PresentImmersive(const MetalRHIVisionOS::PresentImmersiveParams& VisionOSParams)
+void FMetalViewport::PresentImmersive(const MetalRHIVisionOS::PresentImmersiveParams* InVisionOSParams)
 {
+	// The null param case means that we are not really submitting a frame to the compositor.
+	if (InVisionOSParams == nullptr)
+	{
+		FScopeLock Lock(&Mutex);
+		
+		dispatch_semaphore_t& FrameSemaphore = Device.GetFrameSemaphore();
+		dispatch_semaphore_signal(FrameSemaphore);
+		return;
+	}
+	
+	const MetalRHIVisionOS::PresentImmersiveParams& VisionOSParams = *InVisionOSParams;
+	
 	check(SwiftLayer);  // If no SwiftLayer we should not be trying to be immersive.
 	check(VisionOSParams.SwiftFrame);
 
-	FRHICommandListImmediate& RHICmdList = FRHICommandListImmediate::Get();
-	FMetalRHICommandContext& Context = FMetalRHICommandContext::Get(RHICmdList);
+	check(VisionOSParams.RHICommandContext);
+	FMetalRHICommandContext& Context = *static_cast<FMetalRHICommandContext*>(VisionOSParams.RHICommandContext);
 	
 	FScopeLock Lock(&Mutex);
 	
@@ -648,8 +660,6 @@ void FMetalViewport::PresentImmersive(const MetalRHIVisionOS::PresentImmersivePa
 		MTLTexturePtr DrawableTexture = NS::RetainPtr(DrawableTextureParam);
 		MTLTexturePtr DrawableDepthTexture = NS::RetainPtr(DrawableDepthTextureParam);
 		{
-			FScopeLock BlockLock(&Mutex);
-			
 			if (DrawableTexture)
 			{
 				// TODO Currently we are using intermediate back buffer to connect the OXRVisionOS Swapchain to the drawable.
@@ -681,15 +691,28 @@ void FMetalViewport::PresentImmersive(const MetalRHIVisionOS::PresentImmersivePa
 						Context.CopyFromTextureToTexture(Src.get(), 0, 0, MTL::Origin(0, 0, 0), MTL::Size(Width, Height, 1), Dst.get(), 0, 0, MTL::Origin(0, 0, 0));
 					}
 				}
+				
 				FMetalCommandBuffer* Buffer = Context.Finalize();
 				Context.GetDevice().GetCommandQueue().CommitCommandBuffer(Buffer);
 				Context.ResetContext();
 				
-				cp_frame_t CompositorServicesFrame;
-				
 				FMetalCommandBuffer* CurrentCommandBuffer = Context.GetCurrentCommandBuffer();
+				
+				{
+					dispatch_semaphore_t& FrameSemaphore = Device.GetFrameSemaphore();
+					dispatch_retain(FrameSemaphore);
+					MTL::HandlerFunction CommandBufferHandler = [FrameSemaphore](MTL::CommandBuffer* cmd_buf)
+					{
+						dispatch_semaphore_signal(FrameSemaphore);
+						dispatch_release(FrameSemaphore);
+						
+						FMetalGPUProfiler::RecordPresent(cmd_buf);
+					};
+					CurrentCommandBuffer->GetMTLCmdBuffer()->addCompletedHandler(CommandBufferHandler);
+				}
+				
 				cp_drawable_encode_present(VisionOSParams.SwiftDrawable, (__bridge id<MTLCommandBuffer>)CurrentCommandBuffer->GetMTLCmdBuffer().get());
-				CompositorServicesFrame = VisionOSParams.SwiftFrame;
+				cp_frame_t CompositorServicesFrame = VisionOSParams.SwiftFrame;
 				
 				Buffer = Context.Finalize();
 				Context.GetDevice().GetCommandQueue().CommitCommandBuffer(Buffer);
@@ -699,6 +722,8 @@ void FMetalViewport::PresentImmersive(const MetalRHIVisionOS::PresentImmersivePa
 			}
 		}
 	}
+	
+	FMetalGPUProfiler::ResetFrameBufferTimings();
 }
 #endif //PLATFORM_VISIONOS
 
diff --git a/Engine/Source/Runtime/Apple/MetalRHI/Public/MetalRHIVisionOSBridge.h b/Engine/Source/Runtime/Apple/MetalRHI/Public/MetalRHIVisionOSBridge.h
--- a/Engine/Source/Runtime/Apple/MetalRHI/Public/MetalRHIVisionOSBridge.h
+++ b/Engine/Source/Runtime/Apple/MetalRHI/Public/MetalRHIVisionOSBridge.h
@@ -19,7 +19,8 @@ namespace MetalRHIVisionOS
 		cp_frame_t& SwiftFrame;
         cp_drawable_t& SwiftDrawable;
 		int FrameCounter;
+		class IRHICommandContext* RHICommandContext;
     };
-    METALRHI_API void PresentImmersive(const PresentImmersiveParams& Params);
+    METALRHI_API void PresentImmersive(const PresentImmersiveParams* Params);
 }
 #endif // PLATFORM_VISIONOS
diff --git a/Engine/Source/Runtime/Apple/MetalRHI/Public/MetalViewport.h b/Engine/Source/Runtime/Apple/MetalRHI/Public/MetalViewport.h
--- a/Engine/Source/Runtime/Apple/MetalRHI/Public/MetalViewport.h
+++ b/Engine/Source/Runtime/Apple/MetalRHI/Public/MetalViewport.h
@@ -9,6 +9,7 @@
 #include "MetalRHIPrivate.h"
 #include "PixelFormat.h"
 #include "RHIResources.h"
+#include "MetalResources.h"
 
 #if PLATFORM_MAC
 #include "Mac/CocoaTextView.h"
@@ -80,7 +81,7 @@ public:
 	
 #if PLATFORM_VISIONOS
 	void GetDrawableImmersiveTextures(EMetalViewportAccessFlag Accessor, cp_drawable_t SwiftDrawable, MTL::Texture*& OutColorTexture, MTL::Texture*& OutDepthTexture );
-    void PresentImmersive(const MetalRHIVisionOS::PresentImmersiveParams& Params);
+    void PresentImmersive(const MetalRHIVisionOS::PresentImmersiveParams* Params);
 #endif
 	
 private:
