VisionOS - fixing RHIEndDrawingViewport assert and LOCAL_FLOOR space problem
-If the current encoder does not currently have a command buffer we need to start one.
-Added support for LOCAL_FLOOR space.  The default space in visionOS does put zero at the floor.

diff --git a/Engine/Platforms/VisionOS/Plugins/Runtime/OpenXRVisionOS/Source/OXRVisionOS/Private/OXRVisionOSSession.cpp b/Engine/Platforms/VisionOS/Plugins/Runtime/OpenXRVisionOS/Source/OXRVisionOS/Private/OXRVisionOSSession.cpp
--- a/Engine/Platforms/VisionOS/Plugins/Runtime/OpenXRVisionOS/Source/OXRVisionOS/Private/OXRVisionOSSession.cpp
+++ b/Engine/Platforms/VisionOS/Plugins/Runtime/OpenXRVisionOS/Source/OXRVisionOS/Private/OXRVisionOSSession.cpp
@@ -227,7 +227,7 @@ XrResult FOXRVisionOSSession::XrEnumerateReferenceSpaces(
 		return XrResult::XR_ERROR_VALIDATION_FAILURE;
 	}
 
-	const int NumSupportedSpaceTypes = 2;//3;
+	const int NumSupportedSpaceTypes = 3;//4;
 
 	*SpaceCountOutput = NumSupportedSpaceTypes;
 	if (SpaceCapacityInput != 0)
@@ -239,7 +239,8 @@ XrResult FOXRVisionOSSession::XrEnumerateReferenceSpaces(
 
 		OutSpaces[0] = XR_REFERENCE_SPACE_TYPE_VIEW;
 		OutSpaces[1] = XR_REFERENCE_SPACE_TYPE_LOCAL;
-		//OutSpaces[2] = XR_REFERENCE_SPACE_TYPE_STAGE;
+		OutSpaces[2] = XR_REFERENCE_SPACE_TYPE_LOCAL_FLOOR;
+		//OutSpaces[3] = XR_REFERENCE_SPACE_TYPE_STAGE;
 	}
 	return XrResult::XR_SUCCESS;
 }
diff --git a/Engine/Platforms/VisionOS/Plugins/Runtime/OpenXRVisionOS/Source/OXRVisionOS/Private/OXRVisionOSSpace.cpp b/Engine/Platforms/VisionOS/Plugins/Runtime/OpenXRVisionOS/Source/OXRVisionOS/Private/OXRVisionOSSpace.cpp
--- a/Engine/Platforms/VisionOS/Plugins/Runtime/OpenXRVisionOS/Source/OXRVisionOS/Private/OXRVisionOSSpace.cpp
+++ b/Engine/Platforms/VisionOS/Plugins/Runtime/OpenXRVisionOS/Source/OXRVisionOS/Private/OXRVisionOSSpace.cpp
@@ -258,6 +258,19 @@ void FOXRVisionOSSpace::GetTransform(XrTime DisplayTime, FTransform& OutTransfor
 			OutAngularVelocity = FVector::ZeroVector;
 			return;
 		}
+		case XR_REFERENCE_SPACE_TYPE_LOCAL_FLOOR:
+		{
+			// Local floor space is tracking space, so it is always at the origin + any pose it was given and velocity zero.
+			OutLocationFlags = XR_SPACE_LOCATION_ORIENTATION_VALID_BIT
+				| XR_SPACE_LOCATION_POSITION_VALID_BIT
+				| XR_SPACE_LOCATION_ORIENTATION_TRACKED_BIT
+				| XR_SPACE_LOCATION_POSITION_TRACKED_BIT;
+			OutTransform = UEPose;
+			OutVelocityFlags = XR_SPACE_VELOCITY_LINEAR_VALID_BIT | XR_SPACE_VELOCITY_ANGULAR_VALID_BIT;
+			OutLinearVelocity = FVector::ZeroVector;
+			OutAngularVelocity = FVector::ZeroVector;
+			return;
+		}
 //		case XR_REFERENCE_SPACE_TYPE_STAGE:
 //		{
 //			// The stage is the space of the reference space, aka the playable area.  See also XrGetReferenceSpaceBoundsRect.
diff --git a/Engine/Source/Runtime/Apple/MetalRHI/Private/MetalViewport.cpp b/Engine/Source/Runtime/Apple/MetalRHI/Private/MetalViewport.cpp
--- a/Engine/Source/Runtime/Apple/MetalRHI/Private/MetalViewport.cpp
+++ b/Engine/Source/Runtime/Apple/MetalRHI/Private/MetalViewport.cpp
@@ -853,6 +853,10 @@ void FMetalRHICommandContext::RHIEndDrawingViewport(FRHIViewport* ViewportRHI,bo
                 SetCustomPresentViewport(nullptr);
 			}
 			
+			if (!CurrentEncoder.GetCommandBuffer())
+			{
+				StartCommandBuffer();
+			}
 			FMetalCommandBuffer* CurrentCommandBuffer = CurrentEncoder.GetCommandBuffer();
 			check(CurrentCommandBuffer && CurrentCommandBuffer->GetMTLCmdBuffer());
 			
