VisionOS fix for 'double vision' - failure to update hmd pose properly
-Recent refactoring reordered a call such that OXRVisionOS was updating its hmd pose after we fetch it instead of before.
-Fixed it by adding an additional call in the location required by OXRVisionOS.

diff --git a/Engine/Platforms/VisionOS/Plugins/Runtime/OpenXRVisionOS/Source/OXRVisionOS/Private/OXRVisionOS.cpp b/Engine/Platforms/VisionOS/Plugins/Runtime/OpenXRVisionOS/Source/OXRVisionOS/Private/OXRVisionOS.cpp
--- a/Engine/Platforms/VisionOS/Plugins/Runtime/OpenXRVisionOS/Source/OXRVisionOS/Private/OXRVisionOS.cpp
+++ b/Engine/Platforms/VisionOS/Plugins/Runtime/OpenXRVisionOS/Source/OXRVisionOS/Private/OXRVisionOS.cpp
@@ -852,14 +852,14 @@ void FOXRVisionOS::OnBeginRendering_GameThread(XrSession InSession, FSceneViewFa
     Session->OnBeginRendering_GameThread();
 }
 
-void FOXRVisionOS::OnBeginRendering_RenderThread(XrSession InSession, FRDGBuilder& GraphBuilder)
+void FOXRVisionOS::OnBeginRendering_RenderThread_PreDeviceLocationUpdate(XrSession InSession, FRDGBuilder& GraphBuilder)
 {
 	if (!FOXRVisionOSInstance::CheckSession(InSession))
 	{
 		return;
 	}
 	FOXRVisionOSSession* Session = (FOXRVisionOSSession*)InSession;
-	Session->OnBeginRendering_RenderThread();
+	Session->OnBeginRendering_RenderThread_PreDeviceLocationUpdate();
 }
 
 FOXRVisionOS* FOXRVisionOS::Get()
diff --git a/Engine/Platforms/VisionOS/Plugins/Runtime/OpenXRVisionOS/Source/OXRVisionOS/Private/OXRVisionOS.h b/Engine/Platforms/VisionOS/Plugins/Runtime/OpenXRVisionOS/Source/OXRVisionOS/Private/OXRVisionOS.h
--- a/Engine/Platforms/VisionOS/Plugins/Runtime/OpenXRVisionOS/Source/OXRVisionOS/Private/OXRVisionOS.h
+++ b/Engine/Platforms/VisionOS/Plugins/Runtime/OpenXRVisionOS/Source/OXRVisionOS/Private/OXRVisionOS.h
@@ -61,7 +61,7 @@ public:
 	virtual class FOpenXRRenderBridge* GetCustomRenderBridge(XrInstance InInstance) override;
 	virtual bool RequiresRHIContext() const override { return true; }
     virtual void OnBeginRendering_GameThread(XrSession InSession, FSceneViewFamily& InViewFamily, TArrayView<const uint32> VisibleLayers) override;
-	virtual void OnBeginRendering_RenderThread(XrSession InSession, FRDGBuilder& GraphBuilder) override;
+	virtual void OnBeginRendering_RenderThread_PreDeviceLocationUpdate(XrSession InSession, FRDGBuilder& GraphBuilder) override;
 
 	/************************************************************************/
 	/* FOXRVisionOS                                                             */
diff --git a/Engine/Platforms/VisionOS/Plugins/Runtime/OpenXRVisionOS/Source/OXRVisionOS/Private/OXRVisionOSSession.cpp b/Engine/Platforms/VisionOS/Plugins/Runtime/OpenXRVisionOS/Source/OXRVisionOS/Private/OXRVisionOSSession.cpp
--- a/Engine/Platforms/VisionOS/Plugins/Runtime/OpenXRVisionOS/Source/OXRVisionOS/Private/OXRVisionOSSession.cpp
+++ b/Engine/Platforms/VisionOS/Plugins/Runtime/OpenXRVisionOS/Source/OXRVisionOS/Private/OXRVisionOSSession.cpp
@@ -991,9 +991,9 @@ void FOXRVisionOSSession::OnBeginRendering_GameThread()
 }
 
 // FOpenXRHMD::OnBeginRendering_RenderThread
-void FOXRVisionOSSession::OnBeginRendering_RenderThread()
+void FOXRVisionOSSession::OnBeginRendering_RenderThread_PreDeviceLocationUpdate()
 {
-	UE_LOG(LogOXRVisionOS, VeryVerbose, TEXT("%s FC_%i SLF_0x%x FOXRVisionOSSession::OnBeginRendering_RenderThread"), VOSThreadString(), PipelinedFrameStateRendering.FrameCounter, PipelinedFrameStateRendering.SwiftFrame);
+	UE_LOG(LogOXRVisionOS, VeryVerbose, TEXT("%s FC_%i SLF_0x%x FOXRVisionOSSession::OnBeginRendering_RenderThread_PreDeviceLocationUpdate"), VOSThreadString(), PipelinedFrameStateRendering.FrameCounter, PipelinedFrameStateRendering.SwiftFrame);
 	
 	WaitUntil();
     StartSubmission();
diff --git a/Engine/Platforms/VisionOS/Plugins/Runtime/OpenXRVisionOS/Source/OXRVisionOS/Private/OXRVisionOSSession.h b/Engine/Platforms/VisionOS/Plugins/Runtime/OpenXRVisionOS/Source/OXRVisionOS/Private/OXRVisionOSSession.h
--- a/Engine/Platforms/VisionOS/Plugins/Runtime/OpenXRVisionOS/Source/OXRVisionOS/Private/OXRVisionOSSession.h
+++ b/Engine/Platforms/VisionOS/Plugins/Runtime/OpenXRVisionOS/Source/OXRVisionOS/Private/OXRVisionOSSession.h
@@ -122,7 +122,7 @@ public:
 
 	void OnBeginRendering_GameThread();
 
-	void OnBeginRendering_RenderThread();
+	void OnBeginRendering_RenderThread_PreDeviceLocationUpdate();
 
 	void OnBeginFrame_RHIThread(FRHICommandListImmediate* InRHICommandList);
 
diff --git a/Engine/Plugins/Runtime/OpenXR/Source/OpenXRHMD/Private/OpenXRHMD.cpp b/Engine/Plugins/Runtime/OpenXR/Source/OpenXRHMD/Private/OpenXRHMD.cpp
--- a/Engine/Plugins/Runtime/OpenXR/Source/OpenXRHMD/Private/OpenXRHMD.cpp
+++ b/Engine/Plugins/Runtime/OpenXR/Source/OpenXRHMD/Private/OpenXRHMD.cpp
@@ -3293,7 +3293,12 @@ void FOpenXRHMD::OnBeginRendering_RenderThread(FRDGBuilder& GraphBuilder, FScene
 		// Frame submission is not necessary in a headless session.
 		return;
 	}
-
+	
+	for (IOpenXRExtensionPlugin* Module : ExtensionPlugins)
+	{
+		Module->OnBeginRendering_RenderThread_PreDeviceLocationUpdate(Session, GraphBuilder);
+	}
+	
 	// Snapshot new poses for late update. We either do this here, or queue it from OnBeginRendering_GameThread().
 	// If we do it here, it's guaranteed that all platforms will have late update poses available,
 	// but planar reflections will be rendered with pre-late update poses, causing them to visually lag behind the rest of the scene.
diff --git a/Engine/Plugins/Runtime/OpenXR/Source/OpenXRHMD/Public/IOpenXRExtensionPlugin.h b/Engine/Plugins/Runtime/OpenXR/Source/OpenXRHMD/Public/IOpenXRExtensionPlugin.h
--- a/Engine/Plugins/Runtime/OpenXR/Source/OpenXRHMD/Public/IOpenXRExtensionPlugin.h
+++ b/Engine/Plugins/Runtime/OpenXR/Source/OpenXRHMD/Public/IOpenXRExtensionPlugin.h
@@ -382,6 +382,10 @@ public:
 		OnBeginRendering_GameThread(InSession);
 		PRAGMA_ENABLE_DEPRECATION_WARNINGS
 	}
+	
+	virtual void OnBeginRendering_RenderThread_PreDeviceLocationUpdate(XrSession InSession, FRDGBuilder& GraphBuilder)
+	{
+	}
 
 	// OpenXRHMD::OnBeginRendering_RenderThread
 	UE_DEPRECATED(5.6, "Use the FRDGBuilder overload instead")
