visionos cp_frame_query_drawable null handling
-Need to abandon the frame if the drawable is null, this indicates that the app is either being a paused or exited.  Essentially we just continue the frame but don't call the cp_ functions.  At the beginning of the next frame we should be able to handle the state change.
-Note: I never actually replicated the case where null is returned after we attempt to handle the state change, but I do think it is possible.  Isimulated it by nulling out my drawable after calling cp_frame_query_drawable.

diff --git a/Engine/Platforms/VisionOS/Plugins/Runtime/OpenXRVisionOS/Source/OXRVisionOS/Private/OXRVisionOSSession.cpp b/Engine/Platforms/VisionOS/Plugins/Runtime/OpenXRVisionOS/Source/OXRVisionOS/Private/OXRVisionOSSession.cpp
--- a/Engine/Platforms/VisionOS/Plugins/Runtime/OpenXRVisionOS/Source/OXRVisionOS/Private/OXRVisionOSSession.cpp
+++ b/Engine/Platforms/VisionOS/Plugins/Runtime/OpenXRVisionOS/Source/OXRVisionOS/Private/OXRVisionOSSession.cpp
@@ -1097,43 +1097,49 @@ void FOXRVisionOSSession::StartSubmission()
 	{
 		UE_LOG(LogOXRVisionOS, VeryVerbose, TEXT("%s FC_%i SLF_0x%x SD_0x%8x FOXRVisionOSSession::StartSubmission cp_frame_start_submission"), VOSThreadString(), FrameState.FrameCounter, FrameState.SwiftFrame, FrameState.SwiftDrawable);
 		
-		cp_frame_start_submission(FrameState.SwiftFrame);
-		
-		FrameState.SwiftDrawable = cp_frame_query_drawable(FrameState.SwiftFrame);
-    	check(FrameState.SwiftDrawable);
-		
-		UE_LOG(LogOXRVisionOS, VeryVerbose, TEXT("%s FC_%i SLF_0x%x SD_0x%8x FOXRVisionOSSession::StartSubmission cp_frame_query_drawable in got the SD again."), VOSThreadString(), FrameState.FrameCounter, FrameState.SwiftFrame, FrameState.SwiftDrawable);
-		
-		FrameState.SwiftFinalFrameTiming = cp_drawable_get_frame_timing(FrameState.SwiftDrawable);
-		FrameState.SwiftFinalFrameTimeInterval = cp_time_to_cf_time_interval(cp_frame_timing_get_presentation_time(FrameState.SwiftFinalFrameTiming));
-		FrameState.DeviceAnchor = ar_device_anchor_create();
-		
-		UE_LOG(LogOXRVisionOS, VeryVerbose, TEXT("%s FC_%i SLF_0x%x SD_0x%8x FOXRVisionOSSession::StartSubmission cp_drawable_get_frame_timing got SwiftFinalFrameTimeInterval %f"), VOSThreadString(), FrameState.FrameCounter, FrameState.SwiftFrame, FrameState.SwiftDrawable, (double)FrameState.SwiftFinalFrameTimeInterval);
-		
-		uint64_t CurrentMachTime = mach_absolute_time();
-		uint64_t PredictedPresentationMachTime = cp_frame_timing_get_presentation_time(FrameState.SwiftFinalFrameTiming).cp_mach_abs_time;
-		UE_LOG(LogOXRVisionOS, VeryVerbose, TEXT("   FC_%i                   FOXRVisionOSSession::StartSubmission CurrentMachTime %lld PredictedPresentationMachTime %lld"), FrameState.FrameCounter,
-			   CurrentMachTime, PredictedPresentationMachTime);
-		
-		check(ARKitWorldTrackingProvider);
+		FrameState.SwiftDrawable = cp_frame_query_drawable(FrameState.SwiftFrame);	
 		
+		if (FrameState.SwiftDrawable) 
 		{
-			auto anchor_status = ar_world_tracking_provider_query_device_anchor_at_timestamp(ARKitWorldTrackingProvider, FrameState.SwiftFinalFrameTimeInterval, FrameState.DeviceAnchor);
-			if (anchor_status == ar_device_anchor_query_status_success)
-			{
-				cp_drawable_set_device_anchor(FrameState.SwiftDrawable, FrameState.DeviceAnchor);
-				
-				// Get the HMD transform and cache it
-				FrameState.HeadTransform = ar_anchor_get_origin_from_anchor_transform(FrameState.DeviceAnchor);
-			}
-			else
+			cp_frame_start_submission(FrameState.SwiftFrame);				
+			
+			UE_LOG(LogOXRVisionOS, VeryVerbose, TEXT("%s FC_%i SLF_0x%x SD_0x%8x FOXRVisionOSSession::StartSubmission cp_frame_query_drawable in got the SD again."), VOSThreadString(), FrameState.FrameCounter, FrameState.SwiftFrame, FrameState.SwiftDrawable);
+			
+			FrameState.SwiftFinalFrameTiming = cp_drawable_get_frame_timing(FrameState.SwiftDrawable);
+			FrameState.SwiftFinalFrameTimeInterval = cp_time_to_cf_time_interval(cp_frame_timing_get_presentation_time(FrameState.SwiftFinalFrameTiming));
+			FrameState.DeviceAnchor = ar_device_anchor_create();
+			
+			UE_LOG(LogOXRVisionOS, VeryVerbose, TEXT("%s FC_%i SLF_0x%x SD_0x%8x FOXRVisionOSSession::StartSubmission cp_drawable_get_frame_timing got SwiftFinalFrameTimeInterval %f"), VOSThreadString(), FrameState.FrameCounter, FrameState.SwiftFrame, FrameState.SwiftDrawable, (double)FrameState.SwiftFinalFrameTimeInterval);
+			
+			uint64_t CurrentMachTime = mach_absolute_time();
+			uint64_t PredictedPresentationMachTime = cp_frame_timing_get_presentation_time(FrameState.SwiftFinalFrameTiming).cp_mach_abs_time;
+			UE_LOG(LogOXRVisionOS, VeryVerbose, TEXT("   FC_%i                   FOXRVisionOSSession::StartSubmission CurrentMachTime %lld PredictedPresentationMachTime %lld"), FrameState.FrameCounter,
+				   CurrentMachTime, PredictedPresentationMachTime);
+			
+			check(ARKitWorldTrackingProvider);
+			
 			{
-				UE_LOG(LogOXRVisionOS, VeryVerbose, TEXT("%s FC_%i SLF_0x%x SD_0x%8x  FOXRVisionOSSession::StartSubmission cp_drawable_set_device_anchor DeviceAnchor query failed, setting nullptr, no reprojection"), VOSThreadString(), FrameState.FrameCounter, FrameState.SwiftFrame, FrameState.SwiftDrawable);
-				
-				cp_drawable_set_device_anchor(FrameState.SwiftDrawable, nullptr);
-				// Leave FrameState.HeadTransform untouched, same as game thread. It may simply be identity.  But whatever it is we will go ahead and render with it.
+				auto anchor_status = ar_world_tracking_provider_query_device_anchor_at_timestamp(ARKitWorldTrackingProvider, FrameState.SwiftFinalFrameTimeInterval, FrameState.DeviceAnchor);
+				if (anchor_status == ar_device_anchor_query_status_success)
+				{
+					cp_drawable_set_device_anchor(FrameState.SwiftDrawable, FrameState.DeviceAnchor);
+					
+					// Get the HMD transform and cache it
+					FrameState.HeadTransform = ar_anchor_get_origin_from_anchor_transform(FrameState.DeviceAnchor);
+				}
+				else
+				{
+					UE_LOG(LogOXRVisionOS, VeryVerbose, TEXT("%s FC_%i SLF_0x%x SD_0x%8x  FOXRVisionOSSession::StartSubmission cp_drawable_set_device_anchor DeviceAnchor query failed, setting nullptr, no reprojection"), VOSThreadString(), FrameState.FrameCounter, FrameState.SwiftFrame, FrameState.SwiftDrawable);
+					
+					cp_drawable_set_device_anchor(FrameState.SwiftDrawable, nullptr);
+					// Leave FrameState.HeadTransform untouched, same as game thread. It may simply be identity.  But whatever it is we will go ahead and render with it.
+				}
 			}
 		}
+		else 
+		{
+			UE_LOG(LogOXRVisionOS, Log, TEXT("%s FC_%i SLF_0x%x SD_0x%8x FOXRVisionOSSession::StartSubmission doing nothing because SwiftDrawable is null.  This is OK if the app is being paused or exiting."), VOSThreadString(), FrameState.FrameCounter, FrameState.SwiftFrame, FrameState.SwiftDrawable);			
+		}
 	}
 	else
 	{
diff --git a/Engine/Source/Runtime/Apple/MetalRHI/Private/MetalViewport.cpp b/Engine/Source/Runtime/Apple/MetalRHI/Private/MetalViewport.cpp
--- a/Engine/Source/Runtime/Apple/MetalRHI/Private/MetalViewport.cpp
+++ b/Engine/Source/Runtime/Apple/MetalRHI/Private/MetalViewport.cpp
@@ -741,12 +741,13 @@ void FMetalViewport::PresentImmersive(const MetalRHIVisionOS::PresentImmersivePa
 	TRefCountPtr<FMetalSurface> MyLastCompleteFrame = GetMetalSurfaceFromRHITexture(VisionOSParams.Texture);
 	TRefCountPtr<FMetalSurface> MyLastCompleteDepth = GetMetalSurfaceFromRHITexture(VisionOSParams.Depth);
 	{
-		MTL::Texture* DrawableTextureParam = nullptr;
-		MTL::Texture* DrawableDepthTextureParam = nullptr;
-		GetDrawableImmersiveTextures(EMetalViewportAccessDisplayLink, VisionOSParams.SwiftDrawable, DrawableTextureParam, DrawableDepthTextureParam);
-		MTLTexturePtr DrawableTexture = NS::RetainPtr(DrawableTextureParam);
-		MTLTexturePtr DrawableDepthTexture = NS::RetainPtr(DrawableDepthTextureParam);
+		if (VisionOSParams.SwiftDrawable)
 		{
+			MTL::Texture* DrawableTextureParam = nullptr;
+			MTL::Texture* DrawableDepthTextureParam = nullptr;
+			GetDrawableImmersiveTextures(EMetalViewportAccessDisplayLink, VisionOSParams.SwiftDrawable, DrawableTextureParam, DrawableDepthTextureParam);
+			MTLTexturePtr DrawableTexture = NS::RetainPtr(DrawableTextureParam);
+			MTLTexturePtr DrawableDepthTexture = NS::RetainPtr(DrawableDepthTextureParam);
 			if (DrawableTexture)
 			{
 				// TODO Currently we are using intermediate back buffer to connect the OXRVisionOS Swapchain to the drawable.
@@ -765,7 +766,7 @@ void FMetalViewport::PresentImmersive(const MetalRHIVisionOS::PresentImmersivePa
 						
 						Context.CopyFromTextureToTexture(Src.get(), 0, 0, MTL::Origin(0, 0, 0), MTL::Size(Width, Height, 1), Dst.get(), 0, 0, MTL::Origin(0, 0, 0));
 					}
-
+					
 					{
 						TRefCountPtr<FMetalSurface> Texture = MyLastCompleteDepth;
 						check(IsValidRef(Texture));
@@ -778,54 +779,58 @@ void FMetalViewport::PresentImmersive(const MetalRHIVisionOS::PresentImmersivePa
 						Context.CopyFromTextureToTexture(Src.get(), 0, 0, MTL::Origin(0, 0, 0), MTL::Size(Width, Height, 1), Dst.get(), 0, 0, MTL::Origin(0, 0, 0));
 					}
 				}
+			}
+			
+			// We need to make sure that any outstanding encoders have been completed before
+			// We add our completion handler and encode_present.
+			Context.EndCommandBuffer();
+			Context.StartCommandBuffer();
+			
+			// We need to attach the completion handler and the present signal to the final
+			// command buffer
+			FMetalCommandBuffer* FinalCommandBuffer = Context.GetCurrentCommandBuffer();
+			cp_drawable_encode_present(VisionOSParams.SwiftDrawable, (__bridge id<MTLCommandBuffer>)FinalCommandBuffer->GetMTLCmdBuffer());
+		}
 				
-				// We need to make sure that any outstanding encoders have been completed before
-				// We add our completion handler and encode_present.
-				Context.EndCommandBuffer();
-				Context.StartCommandBuffer();
-
-				// We need to attach the completion handler and the present signal to the final
-				// command buffer
-				FMetalCommandBuffer* FinalCommandBuffer = Context.GetCurrentCommandBuffer();;
-				
+		FMetalCommandBuffer* FinalCommandBuffer = Context.GetCurrentCommandBuffer();;
+
 #if ENABLE_METAL_GPUPROFILE && RHI_NEW_GPU_PROFILER == 0
-				FMetalProfiler* Profiler = FMetalProfiler::GetProfiler();
-				FMetalCommandBufferStats* Stats = Profiler->AllocateCommandBuffer(FinalCommandBuffer->GetMTLCmdBuffer(), 0);
+		FMetalProfiler* Profiler = FMetalProfiler::GetProfiler();
+		FMetalCommandBufferStats* Stats = Profiler->AllocateCommandBuffer(FinalCommandBuffer->GetMTLCmdBuffer(), 0);
 #endif
+		
+		{
+			dispatch_semaphore_t& FrameSemaphore = Device.GetFrameSemaphore();
+			dispatch_retain(FrameSemaphore);
+			MTL::HandlerFunction CommandBufferHandler = [FrameSemaphore](MTL::CommandBuffer* cmd_buf)
+			{
+				dispatch_semaphore_signal(FrameSemaphore);
+				dispatch_release(FrameSemaphore);
 				
-				{
-					dispatch_semaphore_t& FrameSemaphore = Device.GetFrameSemaphore();
-					dispatch_retain(FrameSemaphore);
-					MTL::HandlerFunction CommandBufferHandler = [FrameSemaphore](MTL::CommandBuffer* cmd_buf)
-					{
-						dispatch_semaphore_signal(FrameSemaphore);
-						dispatch_release(FrameSemaphore);
-						
 #if RHI_NEW_GPU_PROFILER == 0
-						FMetalCommandBufferTimer::RecordPresent(cmd_buf);
+				FMetalCommandBufferTimer::RecordPresent(cmd_buf);
 #endif
-					};
-					FinalCommandBuffer->GetMTLCmdBuffer()->addCompletedHandler(CommandBufferHandler);
-				}
-				
-				cp_drawable_encode_present(VisionOSParams.SwiftDrawable, (__bridge id<MTLCommandBuffer>)FinalCommandBuffer->GetMTLCmdBuffer());
-				cp_frame_t CompositorServicesFrame = VisionOSParams.SwiftFrame;
+			};
+			FinalCommandBuffer->GetMTLCmdBuffer()->addCompletedHandler(CommandBufferHandler);
+		}
+		
+		METAL_GPUPROFILE(Stats->End(FinalCommandBuffer->GetMTLCmdBuffer()));			
+		
+		TArray<FMetalPayload*> Payloads;
+		Context.Finalize(Payloads);
 
-				METAL_GPUPROFILE(Stats->End(FinalCommandBuffer->GetMTLCmdBuffer()));			
-				
-				TArray<FMetalPayload*> Payloads;
-				Context.Finalize(Payloads);
 
-				FMetalDynamicRHI::Get().SubmitPayloads(MoveTemp(Payloads));
-				
-				Context.ResetContext();
-				
-				cp_frame_end_submission(CompositorServicesFrame);
-				
-				// Wait for the frame semaphore
-				dispatch_semaphore_wait(Device.GetFrameSemaphore(), DISPATCH_TIME_FOREVER);
-			}
+		FMetalDynamicRHI::Get().SubmitPayloads(MoveTemp(Payloads));
+		
+		Context.ResetContext();
+		
+		if (VisionOSParams.SwiftDrawable)
+		{
+			cp_frame_end_submission(VisionOSParams.SwiftFrame);
 		}
+		
+		// Wait for the frame semaphore
+		dispatch_semaphore_wait(Device.GetFrameSemaphore(), DISPATCH_TIME_FOREVER);
 	}
 	
 #if RHI_NEW_GPU_PROFILER == 0
