AVP PresentImmersive avoid use default context.
-Always pass the rhi context in and use it's viewport and device.

diff --git a/Engine/Platforms/VisionOS/Plugins/Runtime/OpenXRVisionOS/Source/OXRVisionOS/Private/OXRVisionOSSession.cpp b/Engine/Platforms/VisionOS/Plugins/Runtime/OpenXRVisionOS/Source/OXRVisionOS/Private/OXRVisionOSSession.cpp
--- a/Engine/Platforms/VisionOS/Plugins/Runtime/OpenXRVisionOS/Source/OXRVisionOS/Private/OXRVisionOSSession.cpp
+++ b/Engine/Platforms/VisionOS/Plugins/Runtime/OpenXRVisionOS/Source/OXRVisionOS/Private/OXRVisionOSSession.cpp
@@ -1165,6 +1165,9 @@ XrResult FOXRVisionOSSession::XrEndFrame(
 	{
 		return XrResult::XR_ERROR_SESSION_NOT_RUNNING;
 	}
+	
+	const XrRHIContextEPIC* RHIContextEPIC = OpenXR::FindChainedStructByType<XrRHIContextEPIC>(InFrameEndInfo, (XrStructureType)XR_TYPE_RHI_CONTEXT_EPIC);
+	check(RHIContextEPIC && RHIContextEPIC->RHIContext);
 
 	if (PipelinedFrameStateRHI.bSynchronizing)
 	{
@@ -1177,14 +1180,16 @@ XrResult FOXRVisionOSSession::XrEndFrame(
 		UE_LOG(LogOXRVisionOS, Log, TEXT("FOXRVisionOSSession XrEndFrame cp_layer_renderer_wait_until_running returned"));
 		
 		// We need to call present so that it can signal the FrameSemaphore.
-		MetalRHIVisionOS::PresentImmersive(nullptr);
+		const MetalRHIVisionOS::PresentImmersiveParams Params{MetalRHIVisionOS::EPresentChoice::DoNotPresent, FTextureRHIRef(), FTextureRHIRef(), PipelinedFrameStateRHI.SwiftFrame, PipelinedFrameStateRHI.SwiftDrawable, PipelinedFrameStateRHI.FrameCounter, RHIContextEPIC->RHIContext};
+		MetalRHIVisionOS::PresentImmersive(Params);
 		return XrResult::XR_SUCCESS;
 	}
 	
 	if (PipelinedFrameStateRHI.bSessionLost)
 	{
 		// We need to call present so that it can signal the FrameSemaphore.
-		MetalRHIVisionOS::PresentImmersive(nullptr);
+		const MetalRHIVisionOS::PresentImmersiveParams Params{MetalRHIVisionOS::EPresentChoice::DoNotPresent, FTextureRHIRef(), FTextureRHIRef(), PipelinedFrameStateRHI.SwiftFrame, PipelinedFrameStateRHI.SwiftDrawable, PipelinedFrameStateRHI.FrameCounter, RHIContextEPIC->RHIContext};
+		MetalRHIVisionOS::PresentImmersive(Params);
 		return XrResult::XR_SUCCESS;
 	}
 	
@@ -1228,12 +1233,9 @@ XrResult FOXRVisionOSSession::XrEndFrame(
 
 		const FOXRVisionOSSwapchain& SwapchainDepthData0 = *(reinterpret_cast<FOXRVisionOSSwapchain*>(Depth0->subImage.swapchain));
 		const FOXRVisionOSSwapchain::FSwapchainImage& WaitedDepth = SwapchainDepthData0.GetLastWaitedImage();
-		
-		const XrRHIContextEPIC* RHIContextEPIC = OpenXR::FindChainedStructByType<XrRHIContextEPIC>(InFrameEndInfo, (XrStructureType)XR_TYPE_RHI_CONTEXT_EPIC);
-		check(RHIContextEPIC && RHIContextEPIC->RHIContext);
 
-		const MetalRHIVisionOS::PresentImmersiveParams Params{WaitedImage.Image, WaitedDepth.Image, PipelinedFrameStateRHI.SwiftFrame, PipelinedFrameStateRHI.SwiftDrawable, PipelinedFrameStateRHI.FrameCounter, RHIContextEPIC->RHIContext};
-        MetalRHIVisionOS::PresentImmersive(&Params);
+		const MetalRHIVisionOS::PresentImmersiveParams Params{MetalRHIVisionOS::EPresentChoice::Present, WaitedImage.Image, WaitedDepth.Image, PipelinedFrameStateRHI.SwiftFrame, PipelinedFrameStateRHI.SwiftDrawable, PipelinedFrameStateRHI.FrameCounter, RHIContextEPIC->RHIContext};
+        MetalRHIVisionOS::PresentImmersive(Params);
     }
 
 	NextBackBufferIndex = (NextBackBufferIndex + 1) % OXRVisionOSSessionHelpers::OXRVisionOSBackbufferLength;
diff --git a/Engine/Source/Runtime/Apple/MetalRHI/Private/MetalRHIVisionOSBridge.cpp b/Engine/Source/Runtime/Apple/MetalRHI/Private/MetalRHIVisionOSBridge.cpp
--- a/Engine/Source/Runtime/Apple/MetalRHI/Private/MetalRHIVisionOSBridge.cpp
+++ b/Engine/Source/Runtime/Apple/MetalRHI/Private/MetalRHIVisionOSBridge.cpp
@@ -9,14 +9,15 @@
 
 DEFINE_LOG_CATEGORY(LogMetalVisionOS);
 
-void MetalRHIVisionOS::PresentImmersive(const MetalRHIVisionOS::PresentImmersiveParams* Params)
+void MetalRHIVisionOS::PresentImmersive(const MetalRHIVisionOS::PresentImmersiveParams& Params)
 {
-    FMetalRHICommandContext* RHICommandContext = static_cast<FMetalRHICommandContext*>(RHIGetDefaultContext());
-    check(RHICommandContext);
-    check(RHICommandContext->CustomPresentViewport);
-    FMetalViewport* Viewport = ResourceCast(RHICommandContext->CustomPresentViewport);
-    check(Viewport);
+	FMetalRHICommandContext* RHICommandContext = static_cast<FMetalRHICommandContext*>(Params.RHICommandContext);
+	check(RHICommandContext);
+	check(RHICommandContext->CustomPresentViewport);
 
-    Viewport->PresentImmersive(Params);
+	FMetalViewport* Viewport = ResourceCast(RHICommandContext->CustomPresentViewport);
+	check(Viewport);
+	
+	Viewport->PresentImmersive(Params);
 }
 #endif // PLATFORM_VISIONOS
diff --git a/Engine/Source/Runtime/Apple/MetalRHI/Private/MetalViewport.cpp b/Engine/Source/Runtime/Apple/MetalRHI/Private/MetalViewport.cpp
--- a/Engine/Source/Runtime/Apple/MetalRHI/Private/MetalViewport.cpp
+++ b/Engine/Source/Runtime/Apple/MetalRHI/Private/MetalViewport.cpp
@@ -599,10 +599,11 @@ void FMetalViewport::GetDrawableImmersiveTextures(EMetalViewportAccessFlag Acces
 }
 
 // This is the present for Immersive visionOS, through the OXRVisionOS plugin.
-void FMetalViewport::PresentImmersive(const MetalRHIVisionOS::PresentImmersiveParams* InVisionOSParams)
+void FMetalViewport::PresentImmersive(const MetalRHIVisionOS::PresentImmersiveParams& VisionOSParams)
 {
-	// The null param case means that we are not really submitting a frame to the compositor.
-	if (InVisionOSParams == nullptr)
+	
+	// This case means that we are not really submitting a frame to the compositor.
+	if (VisionOSParams.PresentChoice == MetalRHIVisionOS::EPresentChoice::DoNotPresent)
 	{
 		FScopeLock Lock(&Mutex);
 		
@@ -611,8 +612,6 @@ void FMetalViewport::PresentImmersive(const MetalRHIVisionOS::PresentImmersivePa
 		return;
 	}
 	
-	const MetalRHIVisionOS::PresentImmersiveParams& VisionOSParams = *InVisionOSParams;
-	
 	check(SwiftLayer);  // If no SwiftLayer we should not be trying to be immersive.
 	check(VisionOSParams.SwiftFrame);
 
diff --git a/Engine/Source/Runtime/Apple/MetalRHI/Public/MetalRHIVisionOSBridge.h b/Engine/Source/Runtime/Apple/MetalRHI/Public/MetalRHIVisionOSBridge.h
--- a/Engine/Source/Runtime/Apple/MetalRHI/Public/MetalRHIVisionOSBridge.h
+++ b/Engine/Source/Runtime/Apple/MetalRHI/Public/MetalRHIVisionOSBridge.h
@@ -12,15 +12,22 @@ DECLARE_LOG_CATEGORY_EXTERN(LogMetalVisionOS, Warning, All);
 
 namespace MetalRHIVisionOS
 {
+	enum class EPresentChoice : uint8
+	{
+		Present,
+		DoNotPresent
+	};
+
     METALRHI_API struct PresentImmersiveParams
     {
-        const FTextureRHIRef& Texture;
-		const FTextureRHIRef& Depth;
-		cp_frame_t& SwiftFrame;
-        cp_drawable_t& SwiftDrawable;
-		int FrameCounter;
-		class IRHICommandContext* RHICommandContext;
+		EPresentChoice PresentChoice = EPresentChoice::DoNotPresent;
+        const FTextureRHIRef Texture;
+		const FTextureRHIRef Depth;
+		cp_frame_t SwiftFrame;
+        cp_drawable_t SwiftDrawable;
+		int FrameCounter = -1;
+		class IRHICommandContext* RHICommandContext = nullptr;
     };
-    METALRHI_API void PresentImmersive(const PresentImmersiveParams* Params);
+    METALRHI_API void PresentImmersive(const PresentImmersiveParams& Params);
 }
 #endif // PLATFORM_VISIONOS
diff --git a/Engine/Source/Runtime/Apple/MetalRHI/Public/MetalViewport.h b/Engine/Source/Runtime/Apple/MetalRHI/Public/MetalViewport.h
--- a/Engine/Source/Runtime/Apple/MetalRHI/Public/MetalViewport.h
+++ b/Engine/Source/Runtime/Apple/MetalRHI/Public/MetalViewport.h
@@ -80,7 +80,7 @@ public:
 	
 #if PLATFORM_VISIONOS
 	void GetDrawableImmersiveTextures(EMetalViewportAccessFlag Accessor, cp_drawable_t SwiftDrawable, MTL::Texture*& OutColorTexture, MTL::Texture*& OutDepthTexture );
-    void PresentImmersive(const MetalRHIVisionOS::PresentImmersiveParams* Params);
+    void PresentImmersive(const MetalRHIVisionOS::PresentImmersiveParams& Params);
 #endif
 	
 private:
