OXRVisionOSSwapchain.h include fixed
-The filename case was wrong, I suppose some tool was accidentally case-insensitive and that was fixed, causing this to start failing.

#Synced-CL 50464994

diff --git a/Engine/Platforms/VisionOS/Plugins/Runtime/OpenXRVisionOS/Source/OXRVisionOS/Private/OXRVisionOSSwapchain.cpp b/Engine/Platforms/VisionOS/Plugins/Runtime/OpenXRVisionOS/Source/OXRVisionOS/Private/OXRVisionOSSwapchain.cpp
--- a/Engine/Platforms/VisionOS/Plugins/Runtime/OpenXRVisionOS/Source/OXRVisionOS/Private/OXRVisionOSSwapchain.cpp
+++ b/Engine/Platforms/VisionOS/Plugins/Runtime/OpenXRVisionOS/Source/OXRVisionOS/Private/OXRVisionOSSwapchain.cpp
@@ -1,9 +1,6 @@
 // Copyright Epic Games, Inc. All Rights Reserved.
 
-#include "OXRVisionOSSwapChain.h"
-
-
-
+#include "OXRVisionOSSwapchain.h"
 
 #if PLATFORM_MAC
 #include "Mac/CocoaTextView.h"
