Fixing virtual swapchain creation function for VisoionOS after adding extra dimension parameter to base

#Synced-CL 52444795

diff --git a/Engine/Platforms/VisionOS/Plugins/Runtime/OpenXRVisionOS/Source/OXRVisionOS/Private/OXRVisionOS_RenderBridge.cpp b/Engine/Platforms/VisionOS/Plugins/Runtime/OpenXRVisionOS/Source/OXRVisionOS/Private/OXRVisionOS_RenderBridge.cpp
--- a/Engine/Platforms/VisionOS/Plugins/Runtime/OpenXRVisionOS/Source/OXRVisionOS/Private/OXRVisionOS_RenderBridge.cpp
+++ b/Engine/Platforms/VisionOS/Plugins/Runtime/OpenXRVisionOS/Source/OXRVisionOS/Private/OXRVisionOS_RenderBridge.cpp
@@ -27,9 +27,9 @@ public:
 		return &Binding;
 	}
 
-	virtual FXRSwapChainPtr CreateSwapchain(XrSession InSession, uint8 Format, uint8& OutFormat, uint32 SizeX, uint32 SizeY, uint32 ArraySize, uint32 NumMips, uint32 NumSamples, ETextureCreateFlags CreateFlags, const FClearValueBinding& ClearValueBinding, ETextureCreateFlags AuxiliaryCreateFlags) override final
+	virtual FXRSwapChainPtr CreateSwapchain(XrSession InSession, uint8 Format, uint8& OutFormat, uint32 SizeX, uint32 SizeY, uint32 ArraySize, uint32 NumMips, uint32 NumSamples, ETextureDimension Dimension, ETextureCreateFlags CreateFlags, const FClearValueBinding& ClearValueBinding, ETextureCreateFlags AuxiliaryCreateFlags) override final
 	{
-		return CreateSwapchain_OXRVisionOS(InSession, Format, OutFormat, SizeX, SizeY, ArraySize, NumMips, NumSamples, CreateFlags, ClearValueBinding);
+		return CreateSwapchain_OXRVisionOS(InSession, Format, OutFormat, SizeX, SizeY, ArraySize, NumMips, NumSamples, Dimension, CreateFlags, ClearValueBinding);
 	}
 
 	//virtual bool Support10BitSwapchain() const override { return true; }
diff --git a/Engine/Platforms/VisionOS/Plugins/Runtime/OpenXRVisionOS/Source/OXRVisionOS/Private/OXRVisionOS_Swapchain.cpp b/Engine/Platforms/VisionOS/Plugins/Runtime/OpenXRVisionOS/Source/OXRVisionOS/Private/OXRVisionOS_Swapchain.cpp
--- a/Engine/Platforms/VisionOS/Plugins/Runtime/OpenXRVisionOS/Source/OXRVisionOS/Private/OXRVisionOS_Swapchain.cpp
+++ b/Engine/Platforms/VisionOS/Plugins/Runtime/OpenXRVisionOS/Source/OXRVisionOS/Private/OXRVisionOS_Swapchain.cpp
@@ -23,8 +23,15 @@ static TArray<XrSwapchainImageOXRVisionOS> EnumerateImagesOXRVisionOS(XrSwapchai
 	return Images;
 }
 
-FXRSwapChainPtr CreateSwapchain_OXRVisionOS(XrSession InSession, uint8 Format, uint8& OutActualFormat, uint32 SizeX, uint32 SizeY, uint32 ArraySize, uint32 NumMips, uint32 NumSamples, ETextureCreateFlags CreateFlags, const FClearValueBinding& ClearValueBinding)
+FXRSwapChainPtr CreateSwapchain_OXRVisionOS(XrSession InSession, uint8 Format, uint8& OutActualFormat, uint32 SizeX, uint32 SizeY, uint32 ArraySize, uint32 NumMips, uint32 NumSamples, ETextureDimension Dimension, ETextureCreateFlags CreateFlags, const FClearValueBinding& ClearValueBinding)
 {
+	if (Dimension != ETextureDimension::Texture2D)
+	{
+		UE_LOG(LogHMD, Error, TEXT("Unsupported texture dimension for VisionOS"));
+		ensure(false);
+		return nullptr;
+	}
+
 	// We use UE formats in OXRVisionOS
 	TFunction<uint32(uint8)> ToPlatformFormat = [](uint8 InFormat) { return InFormat; };
 
@@ -35,7 +42,7 @@ FXRSwapChainPtr CreateSwapchain_OXRVisionOS(XrSession InSession, uint8 Format, u
 	}
 
 	OutActualFormat = Format;
-	XrSwapchain Swapchain = FOpenXRSwapchain::CreateSwapchain(InSession, Format, SizeX, SizeY, ArraySize, NumMips, NumSamples, CreateFlags);
+	XrSwapchain Swapchain = FOpenXRSwapchain::CreateSwapchain(InSession, Format, SizeX, SizeY, ArraySize, NumMips, NumSamples, Dimension, CreateFlags);
 	if (!Swapchain)
 	{
 		return nullptr;
diff --git a/Engine/Platforms/VisionOS/Plugins/Runtime/OpenXRVisionOS/Source/OXRVisionOS/Private/OXRVisionOS_Swapchain.h b/Engine/Platforms/VisionOS/Plugins/Runtime/OpenXRVisionOS/Source/OXRVisionOS/Private/OXRVisionOS_Swapchain.h
--- a/Engine/Platforms/VisionOS/Plugins/Runtime/OpenXRVisionOS/Source/OXRVisionOS/Private/OXRVisionOS_Swapchain.h
+++ b/Engine/Platforms/VisionOS/Plugins/Runtime/OpenXRVisionOS/Source/OXRVisionOS/Private/OXRVisionOS_Swapchain.h
@@ -9,4 +9,4 @@
 
 #include <openxr/openxr.h>
 
-FXRSwapChainPtr CreateSwapchain_OXRVisionOS(XrSession InSession, uint8 Format, uint8& OutActualFormat, uint32 SizeX, uint32 SizeY, uint32 ArraySize, uint32 NumMips, uint32 NumSamples, ETextureCreateFlags CreateFlags, const FClearValueBinding& ClearValueBinding);
+FXRSwapChainPtr CreateSwapchain_OXRVisionOS(XrSession InSession, uint8 Format, uint8& OutActualFormat, uint32 SizeX, uint32 SizeY, uint32 ArraySize, uint32 NumMips, uint32 NumSamples, ETextureDimension Dimension, ETextureCreateFlags CreateFlags, const FClearValueBinding& ClearValueBinding);
diff --git a/Engine/Plugins/Runtime/OpenXR/Source/OpenXRHMD/Private/OpenXRHMD_Swapchain.cpp b/Engine/Plugins/Runtime/OpenXR/Source/OpenXRHMD/Private/OpenXRHMD_Swapchain.cpp
--- a/Engine/Plugins/Runtime/OpenXR/Source/OpenXRHMD/Private/OpenXRHMD_Swapchain.cpp
+++ b/Engine/Plugins/Runtime/OpenXR/Source/OpenXRHMD/Private/OpenXRHMD_Swapchain.cpp
@@ -388,6 +388,7 @@ FXRSwapChainPtr CreateSwapchain_D3D12(XrSession InSession, uint8 Format, uint8&
 			break;
 		default:
 			UE_LOG(LogHMD, Error, TEXT("Unsupported texture dimension"));
+			ensure(false);
 			return nullptr;
 		}
 	}
